#ttp 332021 - EDITOR: PUNTABLE: Streched icons on Check Out Assets message box
#proj UE4
#branch UE4
#summary Fixed padding on SPackagesDialog rows
#extra The icon was set to stretch to the column width; it is now centered instead, and the width has been reduced since it was too wide.
#extra The padding on row entiries has also been fixed so that the header and row content aligns correctly.

diff --git a/Engine/Source/Editor/PackagesDialog/Private/SPackagesDialog.cpp b/Engine/Source/Editor/PackagesDialog/Private/SPackagesDialog.cpp
--- a/Engine/Source/Editor/PackagesDialog/Private/SPackagesDialog.cpp
+++ b/Engine/Source/Editor/PackagesDialog/Private/SPackagesDialog.cpp
@@ -16,8 +16,8 @@ namespace SPackagesDialogDefs
 	const FName ColumnID_FileLabel( "File" );
 	const FName ColumnID_TypeLabel( "Type" );
 
-	const float CheckBoxColumnWidth = 28.0f;
-	const float IconColumnWidth = 28.0f;
+	const float CheckBoxColumnWidth = 23.0f;
+	const float IconColumnWidth = 21.0f;
 }
 
 
@@ -336,9 +336,7 @@ TSharedRef<SWidget> SPackagesDialog::GenerateWidgetForItemAndColumn( TSharedPtr<
 	// Choose the icon based on the severity
 	const FSlateBrush* IconBrush = FEditorStyle::GetBrush( *( Item->GetIconName() ) );
 
-	// Set up checkbox paddings
-	const FMargin PackageMargin(10, 0, 0, 0);
-	const FMargin TypeMargin(5, 0, 0, 0);
+	const FMargin RowPadding(3, 0, 0, 0);
 
 	// Extract the type and color for the package
 	FColor PackageColor;
@@ -356,6 +354,7 @@ TSharedRef<SWidget> SPackagesDialog::GenerateWidgetForItemAndColumn( TSharedPtr<
 	{
 		ItemContentWidget = SNew(SHorizontalBox)
 			+SHorizontalBox::Slot()
+			.Padding(RowPadding)
 			[
 				SNew(SCheckBox)
 				.IsChecked(Item.Get(), &FPackageItem::OnGetDisplayCheckState)
@@ -366,6 +365,8 @@ TSharedRef<SWidget> SPackagesDialog::GenerateWidgetForItemAndColumn( TSharedPtr<
 	{
 		ItemContentWidget = SNew(SHorizontalBox)
 			+SHorizontalBox::Slot()
+			.HAlign(HAlign_Center)
+			.VAlign(VAlign_Center)
 			[
 				SNew(SImage)
 				.Image( IconBrush )
@@ -376,6 +377,7 @@ TSharedRef<SWidget> SPackagesDialog::GenerateWidgetForItemAndColumn( TSharedPtr<
 	{
 		ItemContentWidget = SNew(SHorizontalBox)
 			+SHorizontalBox::Slot()
+			.Padding(RowPadding)
 			[
 				SNew(STextBlock)
 				.Text(PackageName)
@@ -386,6 +388,7 @@ TSharedRef<SWidget> SPackagesDialog::GenerateWidgetForItemAndColumn( TSharedPtr<
 	{
 		ItemContentWidget = SNew(SHorizontalBox)
 			+SHorizontalBox::Slot()
+			.Padding(RowPadding)
 			[
 				SNew(STextBlock)
 				.Text(PackageType)
