Fixed crash in save assets dialog when source control not connected.

diff --git a/Engine/Source/Editor/PackagesDialog/Private/SPackagesDialog.h b/Engine/Source/Editor/PackagesDialog/Private/SPackagesDialog.h
--- a/Engine/Source/Editor/PackagesDialog/Private/SPackagesDialog.h
+++ b/Engine/Source/Editor/PackagesDialog/Private/SPackagesDialog.h
@@ -213,7 +213,10 @@ public:
 		FSourceControlStatePtr SourceControlState = SourceControlProvider.GetState(Package, EStateCacheUsage::Use);
 
 		FString CheckedOutBy;
-		SourceControlState->IsCheckedOutOther(&CheckedOutBy);
+		if (SourceControlState.IsValid())
+		{
+			SourceControlState->IsCheckedOutOther(&CheckedOutBy);
+		}
 
 		return CheckedOutBy;
 	}
