Add the package owner name in the save dialog; that allows sorting by owner package name, in case the user wants to partially submit actor files belonging to a single level, for example.

diff --git a/Engine/Source/Editor/PackagesDialog/Private/PackagesDialog.cpp b/Engine/Source/Editor/PackagesDialog/Private/PackagesDialog.cpp
--- a/Engine/Source/Editor/PackagesDialog/Private/PackagesDialog.cpp
+++ b/Engine/Source/Editor/PackagesDialog/Private/PackagesDialog.cpp
@@ -196,13 +196,15 @@ void FPackagesDialogModule::AddPackageItem(UPackage* InPackage, ECheckBoxState I
 	bool InDisabled/*=false*/, FString InIconName/*=TEXT("SavePackages.SCC_DlgNoIcon")*/, FString InIconToolTip/*=TEXT("")*/)
 {
 	FName AssetName;
+	FName OwnerName;
 
 	// Lookup for the first asset in the package
-	ForEachObjectWithPackage(InPackage, [&AssetName](UObject* InnerObject)
+	ForEachObjectWithPackage(InPackage, [&AssetName, &OwnerName](UObject* InnerObject)
 	{
 		if (InnerObject->IsAsset())
 		{
 			AssetName = InnerObject->GetFName();
+			OwnerName = InnerObject->GetOutermostObject()->GetFName();
 			return false;
 		}
 		return true;
@@ -215,7 +217,7 @@ void FPackagesDialogModule::AddPackageItem(UPackage* InPackage, ECheckBoxState I
 	}
 
 	FString FileName = FPaths::ConvertRelativePathToFull(FPackageName::LongPackageNameToFilename(InPackage->GetName()));
-	PackagesDialogWidget.Get()->Add(MakeShareable(new FPackageItem(InPackage, AssetName.ToString(), FileName, InChecked, InDisabled, InIconName, InIconToolTip)));
+	PackagesDialogWidget.Get()->Add(MakeShareable(new FPackageItem(InPackage, AssetName.ToString(), FileName, OwnerName.ToString(), InChecked, InDisabled, InIconName, InIconToolTip)));
 }
 
 /**
diff --git a/Engine/Source/Editor/PackagesDialog/Private/SPackagesDialog.cpp b/Engine/Source/Editor/PackagesDialog/Private/SPackagesDialog.cpp
--- a/Engine/Source/Editor/PackagesDialog/Private/SPackagesDialog.cpp
+++ b/Engine/Source/Editor/PackagesDialog/Private/SPackagesDialog.cpp
@@ -23,6 +23,7 @@ namespace SPackagesDialogDefs
 	const FName ColumnID_CheckBoxLabel( "CheckBox" );
 	const FName ColumnID_IconLabel( "Icon" );
 	const FName ColumnID_AssetLabel( "Asset" );
+	const FName ColumnID_OwnerLabel( "Owner" );
 	const FName ColumnID_PackageLabel( "Package" );
 	const FName ColumnID_TypeLabel( "Type" );
 	const FName ColumnID_CheckedOutByLabel( "CheckedOutBy" );
@@ -192,6 +193,14 @@ void SPackagesDialog::Construct(const FArguments& InArgs)
 		.FillWidth( 5.0f )
 	);
 
+	HeaderRowWidget->AddColumn(
+		SHeaderRow::Column( SPackagesDialogDefs::ColumnID_OwnerLabel )
+		.DefaultLabel( LOCTEXT("OwnerColumnLabel", "Owner" ) )
+		.SortMode( this, &SPackagesDialog::GetColumnSortMode, SPackagesDialogDefs::ColumnID_OwnerLabel )
+		.OnSort( this, &SPackagesDialog::OnColumnSortModeChanged )
+		.FillWidth( 7.0f )
+	);
+
 	HeaderRowWidget->AddColumn(
 		SHeaderRow::Column( SPackagesDialogDefs::ColumnID_PackageLabel )
 		.DefaultLabel( LOCTEXT("FileColumnLabel", "File" ) )
@@ -468,6 +477,7 @@ TSharedRef<SWidget> SPackagesDialog::GenerateWidgetForItemAndColumn( TSharedPtr<
 	}
 
 	const FString AssetName = Item->GetAssetName();
+	const FString OwnerName = Item->GetOwnerName();
 	const FString PackageName = Item->GetPackageName();
 	const FString FileName = Item->GetFileName();
 
@@ -508,6 +518,18 @@ TSharedRef<SWidget> SPackagesDialog::GenerateWidgetForItemAndColumn( TSharedPtr<
 				.IsEnabled(!Item->IsDisabled())
 			];
 	}
+	else if (ColumnID == SPackagesDialogDefs::ColumnID_OwnerLabel)
+	{
+		ItemContentWidget = SNew(SHorizontalBox)
+			+SHorizontalBox::Slot()
+			.Padding(RowPadding)
+			[
+				SNew(STextBlock)
+				.Text(FText::FromString(OwnerName))
+				.ToolTipText(FText::FromString(OwnerName))
+				.IsEnabled(!Item->IsDisabled())
+			];
+	}
 	else if (ColumnID == SPackagesDialogDefs::ColumnID_PackageLabel)
 	{
 		ItemContentWidget = SNew(SHorizontalBox)
@@ -809,6 +831,19 @@ void SPackagesDialog::SortTree()
 				return A->GetAssetName() >= B->GetAssetName(); } );
 		}
 	}
+	else if (SortByColumn == SPackagesDialogDefs::ColumnID_OwnerLabel)
+	{
+		if (SortMode == EColumnSortMode::Ascending)
+		{
+			Items.Sort([](const TSharedPtr<FPackageItem>& A, const TSharedPtr<FPackageItem>& B) {
+				return A->GetOwnerName() < B->GetOwnerName(); } );
+		}
+		else if (SortMode == EColumnSortMode::Descending)
+		{
+			Items.Sort([](const TSharedPtr<FPackageItem>& A, const TSharedPtr<FPackageItem>& B) {
+				return A->GetOwnerName() >= B->GetOwnerName(); } );
+		}
+	}
 	else if (SortByColumn == SPackagesDialogDefs::ColumnID_PackageLabel)
 	{
 		if (SortMode == EColumnSortMode::Ascending)
diff --git a/Engine/Source/Editor/PackagesDialog/Private/SPackagesDialog.h b/Engine/Source/Editor/PackagesDialog/Private/SPackagesDialog.h
--- a/Engine/Source/Editor/PackagesDialog/Private/SPackagesDialog.h
+++ b/Engine/Source/Editor/PackagesDialog/Private/SPackagesDialog.h
@@ -116,11 +116,12 @@ private:
 class FPackageItem : public TSharedFromThis<FPackageItem>
 {
 public:
-	FPackageItem(UPackage* InPackage, const FString& InAssetName, const FString& InFileName, ECheckBoxState InState, bool InDisabled = false, FString InIconName=TEXT("SavePackages.SCC_DlgNoIcon"), FString InIconToolTip=TEXT(""))
+	FPackageItem(UPackage* InPackage, const FString& InAssetName, const FString& InFileName, const FString& InOwnerName, ECheckBoxState InState, bool InDisabled = false, FString InIconName=TEXT("SavePackages.SCC_DlgNoIcon"), FString InIconToolTip=TEXT(""))
 		: Package(InPackage)
 		, AssetName(InAssetName)
 		, PackageName(FPackageName::ObjectPathToPackageName(InPackage->GetName()))
 		, FileName(InFileName)
+		, OwnerName(InOwnerName)
 		, State(InState)
 		, Disabled(InDisabled)
 		, IconName(InIconName)
@@ -227,6 +228,13 @@ public:
 	 */
 	const FString& GetFileName() const { return FileName; }
 
+	/**
+	 * Gets the name of the file owner item
+	 *
+	 * @return the name of the file owner item
+	 */
+	const FString& GetOwnerName() const { return OwnerName; }
+
 	/**
 	 * Gets the icon name of the checkbox item
 	 *
@@ -316,6 +324,7 @@ private:
 	FString AssetName;						// Name of the asset to display
 	FString PackageName;					// Name of the package to display
 	FString FileName;						// Name of the file
+	FString OwnerName;						// Name of the owner of the file
 	ECheckBoxState State;					// The state of the checkbox
 	bool Disabled; 							// if the entry is disabled
 	FString IconName; 						// Name of an icon to show next to the checkbox
