Show the actor label instead of object name in the save dialog for external actors.

diff --git a/Engine/Source/Developer/AssetTools/Private/AssetTypeActions/AssetTypeActions_Actor.h b/Engine/Source/Developer/AssetTools/Private/AssetTypeActions/AssetTypeActions_Actor.h
--- a/Engine/Source/Developer/AssetTools/Private/AssetTypeActions/AssetTypeActions_Actor.h
+++ b/Engine/Source/Developer/AssetTools/Private/AssetTypeActions/AssetTypeActions_Actor.h
@@ -14,4 +14,5 @@ public:
 	virtual FColor GetTypeColor() const override { return FColor(0,232,0); }
 	virtual UClass* GetSupportedClass() const override { return AActor::StaticClass(); }
 	virtual uint32 GetCategories() override { return EAssetTypeCategories::None; }
+	virtual FString GetObjectDisplayName(UObject* Object) const override { return CastChecked<AActor>(Object)->GetActorLabel(); }
 };
diff --git a/Engine/Source/Developer/AssetTools/Public/AssetTypeActions_Base.h b/Engine/Source/Developer/AssetTools/Public/AssetTypeActions_Base.h
--- a/Engine/Source/Developer/AssetTools/Public/AssetTypeActions_Base.h
+++ b/Engine/Source/Developer/AssetTools/Public/AssetTypeActions_Base.h
@@ -25,6 +25,10 @@ class FAssetTypeActions_Base : public IAssetTypeActions
 public:
 
 	// IAssetTypeActions interface
+	virtual FString GetObjectDisplayName(UObject* Object) const override
+	{
+		return Object->GetName();
+	}
 
 	virtual bool HasActions( const TArray<UObject*>& InObjects ) const override
 	{
diff --git a/Engine/Source/Developer/AssetTools/Public/IAssetTypeActions.h b/Engine/Source/Developer/AssetTools/Public/IAssetTypeActions.h
--- a/Engine/Source/Developer/AssetTools/Public/IAssetTypeActions.h
+++ b/Engine/Source/Developer/AssetTools/Public/IAssetTypeActions.h
@@ -88,6 +88,9 @@ public:
 	/** Returns the categories that this asset type appears in. The return value is one or more flags from EAssetTypeCategories.  */
 	virtual uint32 GetCategories() = 0;
 
+	/** Returns the display name for that object. */
+	virtual FString GetObjectDisplayName(UObject* Object) const = 0;
+
 	/** Returns array of sub-menu names that this asset type is parented under in the Asset Creation Context Menu. */
 	virtual const TArray<FText>& GetSubMenus() const = 0;
 
diff --git a/Engine/Source/Editor/PackagesDialog/Private/PackagesDialog.cpp b/Engine/Source/Editor/PackagesDialog/Private/PackagesDialog.cpp
--- a/Engine/Source/Editor/PackagesDialog/Private/PackagesDialog.cpp
+++ b/Engine/Source/Editor/PackagesDialog/Private/PackagesDialog.cpp
@@ -9,6 +9,8 @@
 #include "SPackagesDialog.h"
 #include "Widgets/Views/SListView.h"
 #include "AssetRegistryModule.h"
+#include "IAssetTools.h"
+#include "AssetToolsModule.h"
 
 IMPLEMENT_MODULE( FPackagesDialogModule, PackagesDialog );
 
@@ -198,12 +200,24 @@ void FPackagesDialogModule::AddPackageItem(UPackage* InPackage, ECheckBoxState I
 	FName AssetName;
 	FName OwnerName;
 
+	FAssetToolsModule& AssetToolsModule = FModuleManager::LoadModuleChecked<FAssetToolsModule>(TEXT("AssetTools"));
+
 	// Lookup for the first asset in the package
-	ForEachObjectWithPackage(InPackage, [&AssetName, &OwnerName](UObject* InnerObject)
+	ForEachObjectWithPackage(InPackage, [&AssetName, &OwnerName, &AssetToolsModule](UObject* InnerObject)
 	{
 		if (InnerObject->IsAsset())
 		{
-			AssetName = InnerObject->GetFName();
+			TWeakPtr<IAssetTypeActions> AssetTypeActions = AssetToolsModule.Get().GetAssetTypeActionsForClass(InnerObject->GetClass());
+
+			if (AssetTypeActions.IsValid())
+			{
+				AssetName = *AssetTypeActions.Pin()->GetObjectDisplayName(InnerObject);
+			}
+			else
+			{
+				AssetName = InnerObject->GetFName();
+			}
+
 			OwnerName = InnerObject->GetOutermostObject()->GetFName();
 			return false;
 		}
