Source control and save dialog reskin

diff --git a/Engine/Source/Editor/EditorStyle/Private/StarshipStyle.cpp b/Engine/Source/Editor/EditorStyle/Private/StarshipStyle.cpp
--- a/Engine/Source/Editor/EditorStyle/Private/StarshipStyle.cpp
+++ b/Engine/Source/Editor/EditorStyle/Private/StarshipStyle.cpp
@@ -430,6 +430,8 @@ void FStarshipEditorStyle::FStyle::SetupGeneralStyles()
 
 	Set( "WarningStripe", new IMAGE_BRUSH( "Common/WarningStripe", FVector2D(20,6), FLinearColor::White, ESlateBrushTileType::Horizontal ) );
 
+	Set("RoundedWarning", new FSlateRoundedBoxBrush(FStyleColors::Warning, 4.0f));
+	Set("RoundedError", new FSlateRoundedBoxBrush(FStyleColors::Error, 4.0f));
 
 	Set( "Button.Disabled", new BOX_BRUSH( "Common/Button_Disabled", 8.0f/32.0f ) );
 	
diff --git a/Engine/Source/Editor/EditorWidgets/Private/SWarningOrErrorBox.cpp b/Engine/Source/Editor/EditorWidgets/Private/SWarningOrErrorBox.cpp
new file mode 100644
--- /dev/null
+++ b/Engine/Source/Editor/EditorWidgets/Private/SWarningOrErrorBox.cpp
@@ -0,0 +1,36 @@
+// Copyright Epic Games, Inc. All Rights Reserved.
+
+#include "SWarningOrErrorBox.h"
+#include "Widgets/Text/STextBlock.h"
+#include "Widgets/SBoxPanel.h"
+#include "Widgets/Images/SImage.h"
+
+void SWarningOrErrorBox::Construct(const FArguments& InArgs)
+{
+	SBorder::Construct(SBorder::FArguments()
+		.Padding(16.0f)
+		.ForegroundColor(FAppStyle::Get().GetSlateColor("Colors.White"))
+		.BorderImage(InArgs._MessageStyle == EMessageStyle::Warning ? FAppStyle::Get().GetBrush("RoundedWarning") : FAppStyle::Get().GetBrush("RoundedError"))
+		[
+			SNew(SHorizontalBox)
+			+ SHorizontalBox::Slot()
+			.VAlign(VAlign_Center)
+			.AutoWidth()
+			.Padding(FMargin(0.0f, 0.0f, 16.0f, 0.0f))
+			[
+				SNew(SImage)
+				.DesiredSizeOverride(FVector2D(24,24))
+				.Image(InArgs._MessageStyle == EMessageStyle::Warning ? FAppStyle::Get().GetBrush("Icons.Warning") : FAppStyle::Get().GetBrush("Icons.Error"))
+				.ColorAndOpacity(FSlateColor::UseForeground())
+			]
+			+SHorizontalBox::Slot()
+			.VAlign(VAlign_Center)
+			[
+				SNew(STextBlock)
+				.Text(InArgs._Message)
+				.ColorAndOpacity(FAppStyle::Get().GetSlateColor("Colors.White"))
+				.AutoWrapText(true)
+			]
+		]);
+}
+
diff --git a/Engine/Source/Editor/EditorWidgets/Public/SWarningOrErrorBox.h b/Engine/Source/Editor/EditorWidgets/Public/SWarningOrErrorBox.h
new file mode 100644
--- /dev/null
+++ b/Engine/Source/Editor/EditorWidgets/Public/SWarningOrErrorBox.h
@@ -0,0 +1,25 @@
+// Copyright Epic Games, Inc. All Rights Reserved.
+
+#pragma once
+
+#include "CoreMinimal.h"
+#include "Widgets/Layout/SBorder.h"
+
+enum class EMessageStyle
+{
+	Warning,
+	Error
+};
+
+class EDITORWIDGETS_API SWarningOrErrorBox : public SBorder
+{
+public:
+	SLATE_BEGIN_ARGS(SWarningOrErrorBox)
+		: _MessageStyle(EMessageStyle::Warning)
+	{}
+		SLATE_ATTRIBUTE(FText, Message)
+		SLATE_ARGUMENT(EMessageStyle, MessageStyle)
+	SLATE_END_ARGS()
+
+	void Construct(const FArguments& InArgs);
+};
\ No newline at end of file
diff --git a/Engine/Source/Editor/PackagesDialog/PackagesDialog.Build.cs b/Engine/Source/Editor/PackagesDialog/PackagesDialog.Build.cs
--- a/Engine/Source/Editor/PackagesDialog/PackagesDialog.Build.cs
+++ b/Engine/Source/Editor/PackagesDialog/PackagesDialog.Build.cs
@@ -19,7 +19,8 @@ public class PackagesDialog : ModuleRules
 				"SlateCore",
                 "EditorStyle",
 				"UnrealEd",
-				"SourceControl"
+				"SourceControl",
+				"EditorWidgets"
 			}
 		);
 
diff --git a/Engine/Source/Editor/PackagesDialog/Private/PackagesDialog.cpp b/Engine/Source/Editor/PackagesDialog/Private/PackagesDialog.cpp
--- a/Engine/Source/Editor/PackagesDialog/Private/PackagesDialog.cpp
+++ b/Engine/Source/Editor/PackagesDialog/Private/PackagesDialog.cpp
@@ -14,7 +14,7 @@
 
 IMPLEMENT_MODULE( FPackagesDialogModule, PackagesDialog );
 
-const FVector2D FPackagesDialogModule::DEFAULT_WINDOW_SIZE = FVector2D(600, 400);	// TODO: MIN SIZE SHOULD BE 270, 330 FOR AESTHETICS
+const FVector2D FPackagesDialogModule::DEFAULT_WINDOW_SIZE = FVector2D(640, 492);	
 const FVector2D FPackagesDialogModule::EXTRA_WINDOW_WIDTH = FVector2D(150, 0);
 
 /**
diff --git a/Engine/Source/Editor/PackagesDialog/Private/SPackagesDialog.cpp b/Engine/Source/Editor/PackagesDialog/Private/SPackagesDialog.cpp
--- a/Engine/Source/Editor/PackagesDialog/Private/SPackagesDialog.cpp
+++ b/Engine/Source/Editor/PackagesDialog/Private/SPackagesDialog.cpp
@@ -15,6 +15,7 @@
 #include "IAssetTools.h"
 #include "IAssetTypeActions.h"
 #include "AssetToolsModule.h"
+#include "SWarningOrErrorBox.h"
 
 #define LOCTEXT_NAMESPACE "SPackagesDialog"
 
@@ -28,8 +29,8 @@ namespace SPackagesDialogDefs
 	const FName ColumnID_TypeLabel( "Type" );
 	const FName ColumnID_CheckedOutByLabel( "CheckedOutBy" );
 
-	const float CheckBoxColumnWidth = 23.0f;
-	const float IconColumnWidth = 21.0f;
+	const float CheckBoxColumnWidth = 38.0f;
+	const float IconColumnWidth = 22.0f;
 }
 
 
@@ -101,12 +102,7 @@ bool FPackageItem::GetTypeNameAndColor(FString& OutName, FColor& OutColor) const
 			}
 			else
 			{
-				const FColor EngineBorderColor = ObjectPtr->IsPendingKill() ? FColor::Red : AssetTypeActions->GetTypeColor();
-				OutColor = FColor(						// Copied from ContentBrowserCLR.cpp
-					127 + EngineBorderColor.R / 2,		// Desaturate the colors a bit (GB colors were too.. much)
-					127 + EngineBorderColor.G / 2,
-					127 + EngineBorderColor.B / 2, 
-					200);		// Opacity
+				OutColor = ObjectPtr->IsPendingKill() ? FColor::Red : AssetTypeActions->GetTypeColor();
 				OutName = AssetTypeActions->GetName().ToString();
 			}
 			return true;
@@ -167,9 +163,14 @@ void SPackagesDialog::Construct(const FArguments& InArgs)
 		HeaderRowWidget->AddColumn(
 			SHeaderRow::Column( SPackagesDialogDefs::ColumnID_CheckBoxLabel )
 			[
-				SAssignNew( ToggleSelectedCheckBox, SCheckBox )
-				.IsChecked(this, &SPackagesDialog::GetToggleSelectedState)
-				.OnCheckStateChanged(this, &SPackagesDialog::OnToggleSelectedCheckBox)
+				SNew(SBox)
+				.Padding(FMargin(6,3,6,3))
+				.HAlign(HAlign_Center)
+				[
+					SAssignNew( ToggleSelectedCheckBox, SCheckBox )
+					.IsChecked(this, &SPackagesDialog::GetToggleSelectedState)
+					.OnCheckStateChanged(this, &SPackagesDialog::OnToggleSelectedCheckBox)
+				]
 			]
 			.FixedWidth( SPackagesDialogDefs::CheckBoxColumnWidth )
 		);
@@ -183,7 +184,7 @@ void SPackagesDialog::Construct(const FArguments& InArgs)
 		.SortMode( this, &SPackagesDialog::GetColumnSortMode, SPackagesDialogDefs::ColumnID_IconLabel )
 		.OnSort( this, &SPackagesDialog::OnColumnSortModeChanged )
 		.FixedWidth( SPackagesDialogDefs::IconColumnWidth )
-		);
+	);
 
 	HeaderRowWidget->AddColumn(
 		SHeaderRow::Column( SPackagesDialogDefs::ColumnID_AssetLabel )
@@ -192,14 +193,16 @@ void SPackagesDialog::Construct(const FArguments& InArgs)
 		.OnSort( this, &SPackagesDialog::OnColumnSortModeChanged )
 		.FillWidth( 5.0f )
 	);
+/*
 
 	HeaderRowWidget->AddColumn(
-		SHeaderRow::Column( SPackagesDialogDefs::ColumnID_OwnerLabel )
-		.DefaultLabel( LOCTEXT("OwnerColumnLabel", "Owner" ) )
-		.SortMode( this, &SPackagesDialog::GetColumnSortMode, SPackagesDialogDefs::ColumnID_OwnerLabel )
-		.OnSort( this, &SPackagesDialog::OnColumnSortModeChanged )
-		.FillWidth( 7.0f )
+		SHeaderRow::Column(SPackagesDialogDefs::ColumnID_OwnerLabel)
+		.DefaultLabel(LOCTEXT("OwnerColumnLabel", "Owner"))
+		.SortMode(this, &SPackagesDialog::GetColumnSortMode, SPackagesDialogDefs::ColumnID_OwnerLabel)
+		.OnSort(this, &SPackagesDialog::OnColumnSortModeChanged)
+		.FillWidth(7.0f)
 	);
+*/
 
 	HeaderRowWidget->AddColumn(
 		SHeaderRow::Column( SPackagesDialogDefs::ColumnID_PackageLabel )
@@ -231,24 +234,21 @@ void SPackagesDialog::Construct(const FArguments& InArgs)
 	this->ChildSlot
 	[
 		SNew(SBorder)
-		.BorderImage(FEditorStyle::GetBrush("ToolPanel.GroupBorder"))
+		.BorderImage(FAppStyle::Get().GetBrush("Brushes.Panel"))
+		.Padding(FMargin(16))
 		[
 			SNew(SVerticalBox)
-			+SVerticalBox::Slot() .Padding(10) .AutoHeight()
-			[
-					SNew( STextBlock )
-					.Text( this, &SPackagesDialog::GetMessage )
-				.AutoWrapText( true )
-				]
-			+SVerticalBox::Slot() .Padding(FMargin(10, 0, 10, 10)) .AutoHeight()
+			+SVerticalBox::Slot()
+			.AutoHeight()
+			.Padding(0.0f,0.0f,0.0f,8.0f)
 			[
-				SNew( STextBlock )
-				.Text( this, &SPackagesDialog::GetWarning )
-				.ColorAndOpacity( FLinearColor::Yellow )
-				.AutoWrapText( true )
-				.Visibility( this, &SPackagesDialog::GetWarningVisibility )
+				SNew(STextBlock)
+				.Text(this, &SPackagesDialog::GetMessage)
+				.AutoWrapText(true)
 			]
-			+SVerticalBox::Slot()  .FillHeight(0.8)
+
+			+SVerticalBox::Slot()
+			.FillHeight(0.8)
 			[
 				SAssignNew(ItemListView, SListView< TSharedPtr<FPackageItem> >)
 					.ListItemsSource(&Items)
@@ -258,7 +258,19 @@ void SPackagesDialog::Construct(const FArguments& InArgs)
 					.HeaderRow( HeaderRowWidget )
 					.SelectionMode( ESelectionMode::None )
 			]
-			+SVerticalBox::Slot() .AutoHeight() .Padding(2) .HAlign(HAlign_Right) .VAlign(VAlign_Bottom)
+			+ SVerticalBox::Slot()
+			.Padding(0, 16.0f, 0, 0)
+			.AutoHeight()
+			[
+				SNew(SWarningOrErrorBox)
+				.Visibility(this, &SPackagesDialog::GetWarningVisibility)
+				.Message(this, &SPackagesDialog::GetWarning)
+			]
+			+SVerticalBox::Slot() 
+			.AutoHeight()
+			.Padding(0.0f, 16.0f, 0.0f, 0.0f)
+			.HAlign(HAlign_Right) 
+			.VAlign(VAlign_Bottom)
 			[
 				ButtonsBox.ToSharedRef()
 			]
@@ -298,15 +310,14 @@ void SPackagesDialog::AddButton(TSharedPtr<FPackageButton> Button)
 
 	ButtonsBox->AddSlot()
 	.AutoWidth()
-	.Padding( 2 )
+	.Padding(5, 0)
 	[
 		SNew(SButton) 
-			.ButtonStyle(&FAppStyle::Get(), Button->GetType() == DRT_Save ? "PrimaryButton" : "Button")
+			.ButtonStyle(&FAppStyle::Get(), (Button->GetType() == DRT_Save || Button->GetType() == DRT_CheckOut) ? "PrimaryButton" : "Button")
 			.TextStyle(&FAppStyle::Get(), "DialogButtonText")
 			.Text(Button->GetName())
-			.ContentPadding(FMargin(10, 3))
 			.ToolTipText(Button->GetToolTip())
-			.IsEnabled(Button.Get(), &FPackageButton::IsEnabled )
+			.IsEnabled(Button.Get(), &FPackageButton::IsEnabled)
 			.HAlign(HAlign_Right)
 			.OnClicked(Button.Get(), &FPackageButton::OnButtonClicked)
 	];
@@ -466,15 +477,12 @@ TSharedRef<SWidget> SPackagesDialog::GenerateWidgetForItemAndColumn( TSharedPtr<
 	// Choose the icon based on the severity
 	const FSlateBrush* IconBrush = Item->GetIconName().IsEmpty() ? FStyleDefaults::GetNoBrush() : FEditorStyle::GetBrush(*(Item->GetIconName()));
 
-	const FMargin RowPadding(3, 0, 0, 0);
+	const FMargin RowPadding(3, 3, 3, 3);
 
 	// Extract the type and color for the package
 	FColor PackageColor;
 	FString PackageType;
-	if (Item->GetTypeNameAndColor(PackageType, PackageColor))
-	{
-		PackageType = FString(TEXT("(")) + PackageType + FString(TEXT(")"));
-	}
+	Item->GetTypeNameAndColor(PackageType, PackageColor);
 
 	const FString AssetName = Item->GetAssetName();
 	const FString OwnerName = Item->GetOwnerName();
@@ -487,7 +495,7 @@ TSharedRef<SWidget> SPackagesDialog::GenerateWidgetForItemAndColumn( TSharedPtr<
 	{
 		ItemContentWidget = SNew(SHorizontalBox)
 			+SHorizontalBox::Slot()
-			.Padding(RowPadding)
+			.Padding(FMargin(10, 3, 6, 3))
 			[
 				SNew(SCheckBox)
 				.IsChecked(Item.Get(), &FPackageItem::OnGetDisplayCheckState)
@@ -569,7 +577,6 @@ TSharedRef<SWidget> SPackagesDialog::GenerateWidgetForItemAndColumn( TSharedPtr<
 				.Text(FText::FromString(CheckedOutByString))
 				.ToolTipText(FText::FromString(CheckedOutByString))
 				.IsEnabled(!Item->IsDisabled())
-				.ColorAndOpacity(PackageColor)
 			];
 	}
 
@@ -790,7 +797,7 @@ FText SPackagesDialog::GetWarning() const
 
 EVisibility SPackagesDialog::GetWarningVisibility() const
 {
-	return (Warning.IsEmpty()) ? EVisibility::Collapsed : EVisibility::Visible;
+	return (Warning.IsEmpty()) ? EVisibility::Collapsed : EVisibility::HitTestInvisible;
 }
 
 EColumnSortMode::Type SPackagesDialog::GetColumnSortMode( const FName ColumnId ) const
diff --git a/Engine/Source/Editor/UnrealEd/Private/FileHelpers.cpp b/Engine/Source/Editor/UnrealEd/Private/FileHelpers.cpp
--- a/Engine/Source/Editor/UnrealEd/Private/FileHelpers.cpp
+++ b/Engine/Source/Editor/UnrealEd/Private/FileHelpers.cpp
@@ -4010,7 +4010,7 @@ FEditorFileUtils::EPromptReturnCode FEditorFileUtils::PromptForCheckoutAndSave(c
 		bCheckDirty,
 		bPromptToSave,
 		NSLOCTEXT("PackagesDialogModule", "PackagesDialogTitle", "Save Content"),
-		NSLOCTEXT("PackagesDialogModule", "PackagesDialogMessage", "Select content to save."),
+		NSLOCTEXT("PackagesDialogModule", "PackagesDialogMessage", "Select Content to Save"),
 		OutFailedPackages,
 		bAlreadyCheckedOut,
 		bCanBeDeclined);
diff --git a/Engine/Source/Runtime/SlateCore/Private/Styling/StarshipCoreStyle.cpp b/Engine/Source/Runtime/SlateCore/Private/Styling/StarshipCoreStyle.cpp
--- a/Engine/Source/Runtime/SlateCore/Private/Styling/StarshipCoreStyle.cpp
+++ b/Engine/Source/Runtime/SlateCore/Private/Styling/StarshipCoreStyle.cpp
@@ -829,6 +829,8 @@ void FStarshipCoreStyle::SetupColors(TSharedRef<FStyle>& Style)
 	Style->Set("Colors.AccentWhite", FStyleColors::AccentWhite);
 	Style->Set("Colors.AccentFolder", FStyleColors::AccentFolder);
 
+	Style->Set("Colors.Warning", FStyleColors::Warning);
+	Style->Set("Colors.Error", FStyleColors::Error);
 
 	Style->Set("Brushes.Black", new FSlateColorBrush(FStyleColors::Black));
 	Style->Set("Brushes.Title", new FSlateColorBrush(FStyleColors::Title));
@@ -873,6 +875,8 @@ void FStarshipCoreStyle::SetupColors(TSharedRef<FStyle>& Style)
 	Style->Set("Brushes.AccentWhite", new FSlateColorBrush(FStyleColors::AccentWhite));
 	Style->Set("Brushes.AccentFolder", new FSlateColorBrush(FStyleColors::AccentFolder));
 
+	Style->Set("Brushes.Warning", new FSlateColorBrush(FStyleColors::Warning));
+	Style->Set("Brushes.Error", new FSlateColorBrush(FStyleColors::Error));
 }
 
 void FStarshipCoreStyle::SetupTextStyles(TSharedRef<FStyle>& Style)
