Fix aligment of source control buttons in checkout dialog

diff --git a/Engine/Source/Editor/PackagesDialog/Private/SPackagesDialog.cpp b/Engine/Source/Editor/PackagesDialog/Private/SPackagesDialog.cpp
--- a/Engine/Source/Editor/PackagesDialog/Private/SPackagesDialog.cpp
+++ b/Engine/Source/Editor/PackagesDialog/Private/SPackagesDialog.cpp
@@ -151,6 +151,7 @@ void SPackagesDialog::Construct(const FArguments& InArgs)
 				.ToolTipText(LOCTEXT("ConnectToSourceControl_Tooltip", "Connect to source control to allow source control operations to be performed on content and levels."))
 				.ContentPadding(FMargin(10, 3))
 				.HAlign(HAlign_Right)
+				.VAlign(VAlign_Center)
 				.Visibility(this, &SPackagesDialog::GetConnectToSourceControlVisibility)
 				.OnClicked(this, &SPackagesDialog::OnConnectToSourceControlClicked)
 		];
@@ -319,6 +320,7 @@ void SPackagesDialog::AddButton(TSharedPtr<FPackageButton> Button)
 			.ToolTipText(Button->GetToolTip())
 			.IsEnabled(Button.Get(), &FPackageButton::IsEnabled)
 			.HAlign(HAlign_Right)
+			.VAlign(VAlign_Center)
 			.OnClicked(Button.Get(), &FPackageButton::OnButtonClicked)
 	];
 }
