Editor: Add more context to the type column for save selected assets (for niagara, data assets, and data tables)
- Implemented GetDisplayNameFromAssetData for data assets and data tables to show the display name for their native parent class and row type respectively

diff --git a/Engine/Source/Developer/AssetTools/Private/AssetTypeActions/AssetTypeActions_DataAsset.cpp b/Engine/Source/Developer/AssetTools/Private/AssetTypeActions/AssetTypeActions_DataAsset.cpp
new file mode 100644
--- /dev/null
+++ b/Engine/Source/Developer/AssetTools/Private/AssetTypeActions/AssetTypeActions_DataAsset.cpp
@@ -0,0 +1,26 @@
+// Copyright Epic Games, Inc. All Rights Reserved.
+
+#include "AssetTypeActions/AssetTypeActions_DataAsset.h"
+
+#define LOCTEXT_NAMESPACE "AssetTypeActions"
+
+static const FName NAME_NativeClass(TEXT("NativeClass"));
+
+FText FAssetTypeActions_DataAsset::GetDisplayNameFromAssetData(const FAssetData& AssetData) const
+{
+	if (AssetData.IsValid())
+	{
+		const FAssetDataTagMapSharedView::FFindTagResult NativeClassTag = AssetData.TagsAndValues.FindTag(NAME_NativeClass);
+		if (NativeClassTag.IsSet())
+		{
+			if (UClass* FoundClass = FindObjectSafe<UClass>(nullptr, *NativeClassTag.GetValue()))
+			{
+				return FText::Format(LOCTEXT("DataAssetWithType", "Data Asset ({0})"), FoundClass->GetDisplayNameText());
+			}
+		}
+	}
+
+	return FText::GetEmpty();
+}
+
+#undef LOCTEXT_NAMESPACE
diff --git a/Engine/Source/Developer/AssetTools/Private/AssetTypeActions/AssetTypeActions_DataAsset.h b/Engine/Source/Developer/AssetTools/Private/AssetTypeActions/AssetTypeActions_DataAsset.h
--- a/Engine/Source/Developer/AssetTools/Private/AssetTypeActions/AssetTypeActions_DataAsset.h
+++ b/Engine/Source/Developer/AssetTools/Private/AssetTypeActions/AssetTypeActions_DataAsset.h
@@ -14,4 +14,5 @@ public:
 	virtual FColor GetTypeColor() const override { return FColor(201, 29, 85); }
 	virtual UClass* GetSupportedClass() const override { return UDataAsset::StaticClass(); }
 	virtual uint32 GetCategories() override { return EAssetTypeCategories::Misc; }
+	virtual FText GetDisplayNameFromAssetData(const FAssetData& AssetData) const override;
 };
diff --git a/Engine/Source/Developer/AssetTools/Private/AssetTypeActions/AssetTypeActions_DataTable.cpp b/Engine/Source/Developer/AssetTools/Private/AssetTypeActions/AssetTypeActions_DataTable.cpp
--- a/Engine/Source/Developer/AssetTools/Private/AssetTypeActions/AssetTypeActions_DataTable.cpp
+++ b/Engine/Source/Developer/AssetTools/Private/AssetTypeActions/AssetTypeActions_DataTable.cpp
@@ -13,6 +13,8 @@
 
 #define LOCTEXT_NAMESPACE "AssetTypeActions"
 
+static const FName NAME_RowStructure(TEXT("RowStructure"));
+
 void FAssetTypeActions_DataTable::GetActions(const TArray<UObject*>& InObjects, FToolMenuSection& Section)
 {
 	auto Tables = GetTypedWeakObjectPtrs<UObject>(InObjects);
@@ -229,4 +231,21 @@ void FAssetTypeActions_DataTable::PerformAssetDiff(UObject* OldAsset, UObject* N
 	}
 }
 
+FText FAssetTypeActions_DataTable::GetDisplayNameFromAssetData(const FAssetData& AssetData) const
+{
+	if (AssetData.IsValid())
+	{
+		const FAssetDataTagMapSharedView::FFindTagResult RowStructureTag = AssetData.TagsAndValues.FindTag(NAME_RowStructure);
+		if (RowStructureTag.IsSet())
+		{
+			if (UScriptStruct* FoundStruct = FindObject<UScriptStruct>(ANY_PACKAGE, *RowStructureTag.GetValue(), true))
+			{
+				return FText::Format(LOCTEXT("DataTableWithRowType", "Data Table ({0})"), FoundStruct->GetDisplayNameText());
+			}
+		}
+	}
+
+	return FText::GetEmpty();
+}
+
 #undef LOCTEXT_NAMESPACE
diff --git a/Engine/Source/Developer/AssetTools/Private/AssetTypeActions/AssetTypeActions_DataTable.h b/Engine/Source/Developer/AssetTools/Private/AssetTypeActions/AssetTypeActions_DataTable.h
--- a/Engine/Source/Developer/AssetTools/Private/AssetTypeActions/AssetTypeActions_DataTable.h
+++ b/Engine/Source/Developer/AssetTools/Private/AssetTypeActions/AssetTypeActions_DataTable.h
@@ -18,6 +18,7 @@ public:
 	virtual void OpenAssetEditor( const TArray<UObject*>& InObjects, TSharedPtr<class IToolkitHost> EditWithinLevelEditor = TSharedPtr<IToolkitHost>() ) override;
 	virtual void GetResolvedSourceFilePaths(const TArray<UObject*>& TypeAssets, TArray<FString>& OutSourceFilePaths) const override;
 	virtual void PerformAssetDiff(UObject* Asset1, UObject* Asset2, const struct FRevisionInfo& OldRevision, const struct FRevisionInfo& NewRevision) const override;
+	virtual FText GetDisplayNameFromAssetData(const FAssetData& AssetData) const override;
 	// End IAssetTypeActions
 
 protected:
diff --git a/Engine/Source/Editor/PackagesDialog/Private/SPackagesDialog.cpp b/Engine/Source/Editor/PackagesDialog/Private/SPackagesDialog.cpp
--- a/Engine/Source/Editor/PackagesDialog/Private/SPackagesDialog.cpp
+++ b/Engine/Source/Editor/PackagesDialog/Private/SPackagesDialog.cpp
@@ -79,7 +79,7 @@ bool FPackageItem::HasMultipleAssets() const
 }
 
 
-bool FPackageItem::GetTypeNameAndColor(FString& OutName, FColor& OutColor) const
+bool FPackageItem::GetTypeNameAndColor(FText& OutName, FColor& OutColor) const
 {
 	// Resolve the object belonging to the package and cache.
 	if (!Object.IsValid())
@@ -98,12 +98,17 @@ bool FPackageItem::GetTypeNameAndColor(FString& OutName, FColor& OutColor) const
 			if(HasMultipleAssets())
 			{
 				OutColor = FColor::White;
-				OutName = LOCTEXT("MultipleAssets", "Multiple Assets").ToString();
+				OutName = LOCTEXT("MultipleAssets", "Multiple Assets");
 			}
 			else
 			{
 				OutColor = ObjectPtr->IsPendingKill() ? FColor::Red : AssetTypeActions->GetTypeColor();
-				OutName = AssetTypeActions->GetName().ToString();
+
+				OutName = AssetTypeActions->GetDisplayNameFromAssetData(FAssetData(ObjectPtr));
+				if (OutName.IsEmpty())
+				{
+					OutName = AssetTypeActions->GetName();
+				}
 			}
 			return true;
 		}
@@ -111,7 +116,7 @@ bool FPackageItem::GetTypeNameAndColor(FString& OutName, FColor& OutColor) const
 	// if we do not find any package object, consider the package empty, return a red `Empty Package`
 	else
 	{
-		OutName = FString(TEXT("Empty Package"));
+		OutName = LOCTEXT("NoAssets", "Empty Package");
 		OutColor = FColor(					// Copied from ContentBrowserCLR.cpp
 			127 + FColor::Red.R / 2,		// Desaturate the colors a bit (GB colors were too.. much)
 			127 + FColor::Red.G / 2,
@@ -483,7 +488,7 @@ TSharedRef<SWidget> SPackagesDialog::GenerateWidgetForItemAndColumn( TSharedPtr<
 
 	// Extract the type and color for the package
 	FColor PackageColor;
-	FString PackageType;
+	FText PackageType;
 	Item->GetTypeNameAndColor(PackageType, PackageColor);
 
 	const FString AssetName = Item->GetAssetName();
@@ -559,8 +564,8 @@ TSharedRef<SWidget> SPackagesDialog::GenerateWidgetForItemAndColumn( TSharedPtr<
 			.Padding(RowPadding)
 			[
 				SNew(STextBlock)
-				.Text(FText::FromString(PackageType))
-				.ToolTipText(FText::FromString(PackageType))
+				.Text(PackageType)
+				.ToolTipText(PackageType)
 				.IsEnabled(!Item->IsDisabled())
 				.ColorAndOpacity(PackageColor)
 			];
@@ -871,12 +876,12 @@ void SPackagesDialog::SortTree()
 		if (SortMode == EColumnSortMode::Ascending)
 		{
 			Items.Sort([](const TSharedPtr<FPackageItem>& A, const TSharedPtr<FPackageItem>& B) {
-				return A->GetTypeName() < B->GetTypeName(); } );
+				return A->GetTypeName().CompareTo(B->GetTypeName()) < 0; } );
 		}
 		else if (SortMode == EColumnSortMode::Descending)
 		{
 			Items.Sort([](const TSharedPtr<FPackageItem>& A, const TSharedPtr<FPackageItem>& B) {
-				return A->GetTypeName() >= B->GetTypeName(); } );
+				return A->GetTypeName().CompareTo(B->GetTypeName()) >= 0; } );
 		}
 	}
 	else if (SortByColumn == SPackagesDialogDefs::ColumnID_IconLabel)
diff --git a/Engine/Source/Editor/PackagesDialog/Private/SPackagesDialog.h b/Engine/Source/Editor/PackagesDialog/Private/SPackagesDialog.h
--- a/Engine/Source/Editor/PackagesDialog/Private/SPackagesDialog.h
+++ b/Engine/Source/Editor/PackagesDialog/Private/SPackagesDialog.h
@@ -285,21 +285,21 @@ public:
 	/**
 	 * Gets the type name and color of the package item
 	 *
-	 * @param OutName	FString into which the type name will be placed, or an empty string if type cannot be obtained
+	 * @param OutName	FText into which the type name will be placed, or an empty string if type cannot be obtained
 	 * @param OutColor	FColor into which the type color will be placed
 	 *
 	 * @return Whether the details were successfully fetched.
 	 */
-	bool GetTypeNameAndColor(FString& OutName, FColor& OutColor) const;
+	bool GetTypeNameAndColor(FText& OutName, FColor& OutColor) const;
 
 	/**
 	 * Gets just the type name of the package item
 	 *
 	 * @return Type name of the package item, or an empty string
 	 */
-	FString GetTypeName() const
+	FText GetTypeName() const
 	{
-		FString OutName;
+		FText OutName;
 		FColor OutColor;
 		GetTypeNameAndColor(OutName, OutColor);
 		return OutName;
