Don't count pending kill objects as assets when determining if a package contains multiple assets.

diff --git a/Engine/Source/Editor/PackagesDialog/Private/SPackagesDialog.cpp b/Engine/Source/Editor/PackagesDialog/Private/SPackagesDialog.cpp
--- a/Engine/Source/Editor/PackagesDialog/Private/SPackagesDialog.cpp
+++ b/Engine/Source/Editor/PackagesDialog/Private/SPackagesDialog.cpp
@@ -43,6 +43,8 @@ UObject* FPackageItem::GetPackageObject() const
 		GetObjectsWithPackage(Package, ObjectsInPackage, false);
 		for (UObject* Obj : ObjectsInPackage)
 		{
+			// Don't filter pending kill objects here as we need to determine if the package contains
+			// a single, pending kill object to properly show it in the save dialog.
 			if (Obj->IsAsset() && !IAssetRegistry::Get()->ShouldSkipAsset(Obj))
 			{
 				return Obj;
@@ -62,7 +64,9 @@ bool FPackageItem::HasMultipleAssets() const
 		UObject* FirstObj = nullptr;
 		for (UObject* Obj : ObjectsInPackage)
 		{
-			if (Obj->IsAsset() && !IAssetRegistry::Get()->ShouldSkipAsset(Obj))
+			// Filter pending kill objects here because we don't want the case where a package contains
+			// an actor and a deleted actor to be reported as multiple assets.
+			if (Obj->IsAsset() && !Obj->IsPendingKill() && !IAssetRegistry::Get()->ShouldSkipAsset(Obj))
 			{
 				if(FirstObj == nullptr)
 				{
