Save Package Dialog: Favor non PendingKill as PackageObject if more than one Asset is found in the package (Actor BP Reinstancing)

diff --git a/Engine/Source/Editor/PackagesDialog/Private/SPackagesDialog.cpp b/Engine/Source/Editor/PackagesDialog/Private/SPackagesDialog.cpp
--- a/Engine/Source/Editor/PackagesDialog/Private/SPackagesDialog.cpp
+++ b/Engine/Source/Editor/PackagesDialog/Private/SPackagesDialog.cpp
@@ -37,6 +37,7 @@ namespace SPackagesDialogDefs
 
 UObject* FPackageItem::GetPackageObject() const
 {
+	UObject* FoundObject = nullptr;
 	if ( !FileName.StartsWith(TEXT("/Temp/Untitled")) )
 	{
 		TArray<UObject*> ObjectsInPackage;
@@ -45,13 +46,17 @@ UObject* FPackageItem::GetPackageObject() const
 		{
 			// Don't filter pending kill objects here as we need to determine if the package contains
 			// a single, pending kill object to properly show it in the save dialog.
+			// Still choose non pending kill objects over pending kill objects.
 			if (Obj->IsAsset() && !UE::AssetRegistry::FFiltering::ShouldSkipAsset(Obj))
 			{
-				return Obj;
+				if (!FoundObject || (FoundObject->IsPendingKill() && !Obj->IsPendingKill()))
+				{
+					FoundObject = Obj;
+				}
 			}
 		}
 	}
-	return nullptr;
+	return FoundObject;
 }
 
 bool FPackageItem::HasMultipleAssets() const
