- WorldPartition: Fix leaking ActorDesc when Actor (in this case PartitonedActor) replaces a deleted Actor with the same name
- Allow deleting Foliage actors
- SPackagesDialog: Do not show multiple assets when package contains 1 valid object and 1 or n invalid objects

diff --git a/Engine/Source/Editor/PackagesDialog/Private/SPackagesDialog.cpp b/Engine/Source/Editor/PackagesDialog/Private/SPackagesDialog.cpp
--- a/Engine/Source/Editor/PackagesDialog/Private/SPackagesDialog.cpp
+++ b/Engine/Source/Editor/PackagesDialog/Private/SPackagesDialog.cpp
@@ -47,23 +47,29 @@ UObject* FPackageItem::GetPackageObject() const
 
 bool FPackageItem::HasMultipleAssets() const
 {
-	bool bHasMultipleAssets = false;
 	if ( !FileName.StartsWith(TEXT("/Temp/Untitled")) )
 	{
 		int32 NumAssets = 0;
-		ForEachObjectWithPackage(Package, [&NumAssets](UObject* Obj)
+		int32 NumDeleted = 0;
+		ForEachObjectWithPackage(Package, [&NumAssets,&NumDeleted](UObject* Obj)
 			{
 				if (Obj->IsAsset() && !UE::AssetRegistry::FFiltering::ShouldSkipAsset(Obj))
 				{
 					++NumAssets;
+
+					if (!IsValid(Obj))
+					{
+						++NumDeleted;
+					}
 				}
 				return true;
 			}, false /*bIncludeNestedObjects*/);
-		bHasMultipleAssets = NumAssets > 1;
+
+		return (NumAssets - NumDeleted) > 1;
 	}
-	return bHasMultipleAssets;
-}
 
+	return false;
+}
 
 bool FPackageItem::GetTypeNameAndColor(FText& OutName, FColor& OutColor) const
 {
diff --git a/Engine/Source/Runtime/Engine/Private/ActorPartition/ActorPartitionSubsystem.cpp b/Engine/Source/Runtime/Engine/Private/ActorPartition/ActorPartitionSubsystem.cpp
--- a/Engine/Source/Runtime/Engine/Private/ActorPartition/ActorPartitionSubsystem.cpp
+++ b/Engine/Source/Runtime/Engine/Private/ActorPartition/ActorPartitionSubsystem.cpp
@@ -281,6 +281,9 @@ public:
 				// if the actor is using an external package. We really just want to rename that actor out of the way so we can spawn the new one in
 				// the exact same package, keeping the package name intact.
 				ExistingActor->UObject::Rename(nullptr, nullptr, REN_DontCreateRedirectors | REN_DoNotDirty | REN_NonTransactional | REN_ForceNoResetLoaders);
+				
+				// Reuse ActorGuid so that ActorDesc can be updated on save
+				SpawnParams.OverrideActorGuid = ExistingActor->GetActorGuid();
 			}
 						
 			FVector CellCenter(CellBounds.GetCenter());
diff --git a/Engine/Source/Runtime/Engine/Private/Level.cpp b/Engine/Source/Runtime/Engine/Private/Level.cpp
--- a/Engine/Source/Runtime/Engine/Private/Level.cpp
+++ b/Engine/Source/Runtime/Engine/Private/Level.cpp
@@ -3430,8 +3430,15 @@ TArray<UPackage*> ULevel::GetLoadedExternalObjectPackages() const
 
 UPackage* ULevel::CreateActorPackage(UPackage* InLevelPackage, EActorPackagingScheme ActorPackagingScheme, const FString& InActorPath)
 {
-	UPackage* ActorPackage = CreatePackage(*GetActorPackageName(InLevelPackage, ActorPackagingScheme, InActorPath));
-	ActorPackage->SetPackageFlags(PKG_EditorOnly | PKG_ContainsMapData);
+	const FString PackageName = GetActorPackageName(InLevelPackage, ActorPackagingScheme, InActorPath);
+	if (UPackage* ExistingActorPackage = FindObject<UPackage>(nullptr, *PackageName))
+	{
+		check(ExistingActorPackage->HasAllPackagesFlags(PKG_EditorOnly | PKG_ContainsMapData));
+		return ExistingActorPackage;
+	}
+
+	UPackage* ActorPackage = CreatePackage(*PackageName);
+	ActorPackage->SetPackageFlags(PKG_EditorOnly | PKG_ContainsMapData | PKG_NewlyCreated);
 	return ActorPackage;
 }
 
diff --git a/Engine/Source/Runtime/Engine/Private/LevelActor.cpp b/Engine/Source/Runtime/Engine/Private/LevelActor.cpp
--- a/Engine/Source/Runtime/Engine/Private/LevelActor.cpp
+++ b/Engine/Source/Runtime/Engine/Private/LevelActor.cpp
@@ -740,7 +740,7 @@ AActor* UWorld::SpawnActor( UClass* Class, FTransform const* UserTransformPtr, c
 #if WITH_EDITOR
 	if (GIsEditor)
 	{
-		if (Actor->IsAsset())
+		if (Actor->IsAsset() && Actor->GetPackage()->HasAnyPackageFlags(PKG_NewlyCreated))
 		{
 			FAssetRegistryModule::AssetCreated(Actor);
 		}
diff --git a/Engine/Source/Runtime/Foliage/Public/InstancedFoliageActor.h b/Engine/Source/Runtime/Foliage/Public/InstancedFoliageActor.h
--- a/Engine/Source/Runtime/Foliage/Public/InstancedFoliageActor.h
+++ b/Engine/Source/Runtime/Foliage/Public/InstancedFoliageActor.h
@@ -60,6 +60,7 @@ public:
 
 public:
 #if WITH_EDITOR
+	bool CanDeleteSelectedActor(FText& OutReason) const override { return true; }
 	bool CanEditFoliageInstance(const FFoliageInstanceId& InstanceId) const;
 	bool CanMoveFoliageInstance(const FFoliageInstanceId& InstanceId, const ETypedElementWorldType WorldType) const;
 	bool GetFoliageInstanceTransform(const FFoliageInstanceId& InstanceId, FTransform& OutInstanceTransform, bool bWorldSpace) const;
