Displayed the type of actor in the 'Check Out Files' dialog rather than just displaying "Actor".
  - The type displayed is consistent with the type displayed in the 'Changelists' window.

Removed the 'Off' from 'Source Control Off' on the tool bar because the icon and the text color are enough.

diff --git a/Engine/Source/Editor/PackagesDialog/Private/SPackagesDialog.cpp b/Engine/Source/Editor/PackagesDialog/Private/SPackagesDialog.cpp
--- a/Engine/Source/Editor/PackagesDialog/Private/SPackagesDialog.cpp
+++ b/Engine/Source/Editor/PackagesDialog/Private/SPackagesDialog.cpp
@@ -120,15 +120,12 @@ bool FPackageItem::GetTypeNameAndColor(FText& OutName, FColor& OutColor) const
 				OutColor = FColor::White;
 				OutName = LOCTEXT("MultipleAssets", "Multiple Assets");
 			}
-			else
+			else // Just one asset in the package.
 			{
 				OutColor = !IsValidChecked(ObjectPtr) ? FColor::Red : AssetTypeActions->GetTypeColor();
 
-				OutName = AssetTypeActions->GetDisplayNameFromAssetData(FAssetData(ObjectPtr));
-				if (OutName.IsEmpty())
-				{
-					OutName = AssetTypeActions->GetName();
-				}
+				FAssetData AssetData(ObjectPtr);
+				OutName = FText::FromString(AssetData.AssetClassPath.ToString());
 			}
 			return true;
 		}
diff --git a/Engine/Source/Editor/StatusBar/Private/SourceControlMenuHelpers.cpp b/Engine/Source/Editor/StatusBar/Private/SourceControlMenuHelpers.cpp
--- a/Engine/Source/Editor/StatusBar/Private/SourceControlMenuHelpers.cpp
+++ b/Engine/Source/Editor/StatusBar/Private/SourceControlMenuHelpers.cpp
@@ -208,7 +208,7 @@ FText FSourceControlMenuHelpers::GetSourceControlStatusText()
 		}
 		else
 		{
-			return LOCTEXT("SourceControlStatus_Error_Off", "Source Control Off");
+			return LOCTEXT("SourceControlStatus_Error_Off", "Source Control"); // Relies on the icon on the status bar widget to know if the source control is on or off.
 		}
 	}
 }
