Package Dialog: Add RMB menu options to copy the file path and the package path to the clipboard. Fix the file path to include the extension and use platform native formatting.

diff --git a/Engine/Source/Editor/PackagesDialog/PackagesDialog.Build.cs b/Engine/Source/Editor/PackagesDialog/PackagesDialog.Build.cs
--- a/Engine/Source/Editor/PackagesDialog/PackagesDialog.Build.cs
+++ b/Engine/Source/Editor/PackagesDialog/PackagesDialog.Build.cs
@@ -10,6 +10,7 @@ public class PackagesDialog : ModuleRules
 
 		PrivateDependencyModuleNames.AddRange(
 			new string[] {
+				"ApplicationCore",
 				"Core", 
 				"CoreUObject", 
 				"EditorFramework",
diff --git a/Engine/Source/Editor/PackagesDialog/Private/PackagesDialog.cpp b/Engine/Source/Editor/PackagesDialog/Private/PackagesDialog.cpp
--- a/Engine/Source/Editor/PackagesDialog/Private/PackagesDialog.cpp
+++ b/Engine/Source/Editor/PackagesDialog/Private/PackagesDialog.cpp
@@ -244,7 +244,12 @@ void FPackagesDialogModule::AddPackageItem(UPackage* InPackage, ECheckBoxState I
 		AssetName = *FPackageName::GetShortName(InPackage->GetFName());
 	}
 
-	FString FileName = FPaths::ConvertRelativePathToFull(FPackageName::LongPackageNameToFilename(InPackage->GetName()));
+	const FString Extension = (InPackage->GetPackageFlags() & PKG_ContainsMap) ? FPackageName::GetMapPackageExtension() : FPackageName::GetAssetPackageExtension();
+	FString FileName = FPaths::ConvertRelativePathToFull(FPackageName::LongPackageNameToFilename(InPackage->GetName(), Extension));
+
+	// Ensure that the filename is formatted in a way that the current platform natively understands. 
+	FPaths::MakePlatformFilename(FileName);
+	
 	PackagesDialogWidget.Get()->Add(MakeShareable(new FPackageItem(InPackage, AssetName.ToString(), FileName, OwnerName.ToString(), InChecked, InDisabled, InIconName, InIconToolTip)));
 }
 
diff --git a/Engine/Source/Editor/PackagesDialog/Private/SPackagesDialog.cpp b/Engine/Source/Editor/PackagesDialog/Private/SPackagesDialog.cpp
--- a/Engine/Source/Editor/PackagesDialog/Private/SPackagesDialog.cpp
+++ b/Engine/Source/Editor/PackagesDialog/Private/SPackagesDialog.cpp
@@ -7,6 +7,7 @@
 #include "Framework/Commands/UIAction.h"
 #include "Framework/MultiBox/MultiBoxBuilder.h"
 #include "Framework/Views/ITypedTableView.h"
+#include "HAL/PlatformApplicationMisc.h"
 #include "IAssetTools.h"
 #include "IAssetTypeActions.h"
 #include "Input/Events.h"
@@ -21,6 +22,7 @@
 #include "Modules/ModuleManager.h"
 #include "SWarningOrErrorBox.h"
 #include "SlotBase.h"
+#include "Algo/AnyOf.h"
 #include "Styling/AppStyle.h"
 #include "Styling/ISlateStyle.h"
 #include "Styling/SlateColor.h"
@@ -279,7 +281,7 @@ void SPackagesDialog::Construct(const FArguments& InArgs)
 					.OnContextMenuOpening(this, &SPackagesDialog::MakePackageListContextMenu)
 					.ItemHeight(20)
 					.HeaderRow( HeaderRowWidget )
-					.SelectionMode( ESelectionMode::None )
+					.SelectionMode( ESelectionMode::Single )
 			]
 			+ SVerticalBox::Slot()
 			.Padding(0, 16.0f, 0, 0)
@@ -623,6 +625,51 @@ TSharedPtr<SWidget> SPackagesDialog::MakePackageListContextMenu() const
 					FCanExecuteAction::CreateSP( this, &SPackagesDialog::CanExecuteSCCDiffAgainstDepot )
 				)
 			);	
+			MenuBuilder.AddMenuEntry(
+				LOCTEXT("SCCCopyFilePathToClipboard", "Copy File Path"),
+				LOCTEXT("SCCCopyFilePathToClipboardTooltip", "Copies the file path on disk to the clipboard."),
+				FSlateIcon(),
+				FUIAction(
+					FExecuteAction::CreateLambda( [SelectedItems] ()
+					{
+						TArray<FString> Paths;
+						for (const TSharedPtr<FPackageItem>& PackageItem: SelectedItems)
+						{
+							if (!PackageItem->GetFileName().IsEmpty())
+							{
+								Paths.Add(PackageItem->GetFileName());
+							}
+						}
+						FPlatformApplicationMisc::ClipboardCopy(*FString::Join(Paths, TEXT("\n")));
+					} ),
+					FCanExecuteAction::CreateLambda( [SelectedItems] ()
+					{
+						return Algo::AnyOf(SelectedItems, [](const TSharedPtr<FPackageItem>& PackageItem)
+						{
+							return !PackageItem->GetFileName().IsEmpty();
+						});
+					} )
+				)
+			);	
+			MenuBuilder.AddMenuEntry(
+				LOCTEXT("SCCCopyPackagePathToClipboard", "Copy Package Path"),
+				LOCTEXT("SCCCopyPackagePathToClipboardTooltip", "Copies the package path to the clipboard."),
+				FSlateIcon(),
+				FUIAction(
+					FExecuteAction::CreateLambda( [SelectedItems] ()
+					{
+						TArray<FString> Paths;
+						for (const TSharedPtr<FPackageItem>& PackageItem: SelectedItems)
+						{
+							if (!PackageItem->GetFileName().IsEmpty())
+							{
+								Paths.Add(PackageItem->GetPackageName());
+							}
+						}
+						FPlatformApplicationMisc::ClipboardCopy(*FString::Join(Paths, TEXT("\n")));
+					} )
+				)
+			);	
 		}
 		MenuBuilder.EndSection();
 	}
