Improve checkbox behavior of SPackagesDialog

* Make the header checkbox of SPackagesDialog always check/uncheck all items, instead of the selected items.
* If one or more items are selected, and one of their checkboxes is clicked, then all selected items are checked/unchecked together (reimplementation of the removed behavior of the header checkbox checking/unchecking selected items).
* Move some logic to lambdas to reduce the size of the private interface.

diff --git a/Engine/Source/Editor/PackagesDialog/Private/SPackagesDialog.cpp b/Engine/Source/Editor/PackagesDialog/Private/SPackagesDialog.cpp
--- a/Engine/Source/Editor/PackagesDialog/Private/SPackagesDialog.cpp
+++ b/Engine/Source/Editor/PackagesDialog/Private/SPackagesDialog.cpp
@@ -185,6 +185,7 @@ void SPackagesDialog::Construct(const FArguments& InArgs)
 
 	if (!bReadOnly)
 	{
+		// clang-format off
 		HeaderRowWidget->AddColumn(
 			SHeaderRow::Column( SPackagesDialogDefs::ColumnID_CheckBoxLabel )
 			[
@@ -192,13 +193,68 @@ void SPackagesDialog::Construct(const FArguments& InArgs)
 				.Padding(FMargin(6,3,6,3))
 				.HAlign(HAlign_Center)
 				[
-					SAssignNew( ToggleSelectedCheckBox, SCheckBox )
-					.IsChecked(this, &SPackagesDialog::GetToggleSelectedState)
-					.OnCheckStateChanged(this, &SPackagesDialog::OnToggleSelectedCheckBox)
+					SNew(SCheckBox)
+					.IsChecked_Lambda(
+						[this]()
+						{
+							bool bAllChecked = true;
+							bool bAllUnchecked = true;
+							for (const TSharedPtr<FPackageItem>& Item : Items)
+							{
+								if (Item->GetState() != ECheckBoxState::Checked)
+								{
+									bAllChecked = false;
+								}
+								if (Item->GetState() != ECheckBoxState::Unchecked)
+								{
+									bAllUnchecked = false;
+								}
+							}
+
+							if (bAllChecked)
+							{
+								return ECheckBoxState::Checked;
+							}
+							else if (bAllUnchecked)
+							{
+								return ECheckBoxState::Unchecked;
+							}
+							else
+							{
+								return ECheckBoxState::Undetermined;
+							}
+						}
+					)
+					.OnCheckStateChanged_Lambda(
+						[this](ECheckBoxState InNewState)
+						{
+							for (const TSharedPtr<FPackageItem>& Item : Items)
+							{
+								if (InNewState == ECheckBoxState::Checked)
+								{
+									if (Item->IsDisabled())
+									{
+										Item->SetState(ECheckBoxState::Undetermined);
+									}
+									else
+									{
+										Item->SetState(ECheckBoxState::Checked);
+									}
+								}
+								else
+								{
+									Item->SetState(ECheckBoxState::Unchecked);
+								}
+							}
+
+							ItemListView->RequestListRefresh();
+						}
+					)
 				]
 			]
 			.FixedWidth( SPackagesDialogDefs::CheckBoxColumnWidth )
 		);
+		// clang-format on
 	}
 
 	HeaderRowWidget->AddColumn(
@@ -256,6 +312,7 @@ void SPackagesDialog::Construct(const FArguments& InArgs)
 			);
 	}
 
+	// clang-format off
 	this->ChildSlot
 	[
 		SNew(SBorder)
@@ -280,7 +337,7 @@ void SPackagesDialog::Construct(const FArguments& InArgs)
 					.OnGenerateRow(this, &SPackagesDialog::MakePackageListItemWidget)
 					.OnContextMenuOpening(this, &SPackagesDialog::MakePackageListContextMenu)
 					.HeaderRow( HeaderRowWidget )
-					.SelectionMode( ESelectionMode::Single )
+					.SelectionMode( ESelectionMode::Multi )
 			]
 			+ SVerticalBox::Slot()
 			.Padding(0, 16.0f, 0, 0)
@@ -300,6 +357,7 @@ void SPackagesDialog::Construct(const FArguments& InArgs)
 			]
 		]
 	];
+	// clang-format on
 }
 
 /**
@@ -516,14 +574,40 @@ TSharedRef<SWidget> SPackagesDialog::GenerateWidgetForItemAndColumn( TSharedPtr<
 
 	if (ColumnID == SPackagesDialogDefs::ColumnID_CheckBoxLabel)
 	{
+		// clang-format off
 		ItemContentWidget = SNew(SHorizontalBox)
 			+SHorizontalBox::Slot()
 			.Padding(FMargin(10, 3, 6, 3))
 			[
 				SNew(SCheckBox)
 				.IsChecked(Item.Get(), &FPackageItem::OnGetDisplayCheckState)
-				.OnCheckStateChanged(Item.Get(), &FPackageItem::OnDisplayCheckStateChanged)
+				.OnCheckStateChanged_Lambda(
+					[this, WeakItem = Item.ToWeakPtr()](ECheckBoxState InNewState)
+					{
+						if (TSharedPtr<FPackageItem> PinnedItem = WeakItem.Pin())
+						{
+							const bool bAllIfNone = false;
+							const TArray<TSharedPtr<FPackageItem>> SelectedItems = this->GetSelectedItems(bAllIfNone);
+							const bool bItemIsSelected = SelectedItems.Contains(PinnedItem);
+
+							// If this item is selected, check/uncheck it together with all other checked items.
+							if (bItemIsSelected)
+							{
+								for (const TSharedPtr<FPackageItem>& SelectedItem : SelectedItems)
+								{
+									SelectedItem->SetState(InNewState);
+								}
+							}
+							// If this item is NOT in the current selection, just check/uncheck the current item.
+							else
+							{
+								PinnedItem->SetState(InNewState);
+							}
+						}
+					}
+				)
 			];
+		// clang-format on
 	}
 	else if (ColumnID == SPackagesDialogDefs::ColumnID_IconLabel)
 	{
@@ -714,54 +798,6 @@ TArray< TSharedPtr<FPackageItem> > SPackagesDialog::GetSelectedItems( bool bAllI
 	return SelectedItems;
 }
 
-ECheckBoxState SPackagesDialog::GetToggleSelectedState() const
-{
-	//default to a Checked state
-	ECheckBoxState PendingState = ECheckBoxState::Checked;
-
-	TArray< TSharedPtr<FPackageItem> > SelectedItems = GetSelectedItems( true );
-
-	//Iterate through the list of selected packages
-	for ( auto SelectedItem = SelectedItems.CreateConstIterator(); SelectedItem; ++SelectedItem )
-	{
-		if ( SelectedItem->Get()->GetState() == ECheckBoxState::Unchecked )
-		{
-			//if any package in the selection is Unchecked, then represent the entire set of highlighted packages as Unchecked,
-			//so that the first (user) toggle of ToggleSelectedCheckBox consistently Checks all highlighted packages
-			PendingState = ECheckBoxState::Unchecked;
-		}
-	}
-
-	return PendingState;
-}
-
-void SPackagesDialog::OnToggleSelectedCheckBox(ECheckBoxState InNewState)
-{
-	TArray< TSharedPtr<FPackageItem> > SelectedItems = GetSelectedItems( true );
-
-	for ( auto SelectedItem = SelectedItems.CreateConstIterator(); SelectedItem; ++SelectedItem )
-	{
-		FPackageItem *item = (*SelectedItem).Get();
-		if (InNewState == ECheckBoxState::Checked)
-		{
-			if(item->IsDisabled())
-			{
-				item->SetState(ECheckBoxState::Undetermined);
-			}
-			else
-			{
-				item->SetState(ECheckBoxState::Checked);
-			}
-		}
-		else
-		{
-			item->SetState(ECheckBoxState::Unchecked);
-		}
-	}
-
-	ItemListView->RequestListRefresh();
-}
-
 void SPackagesDialog::Tick(const FGeometry& AllottedGeometry, const double InCurrentTime, const float InDeltaTime)
 {
 	SCompoundWidget::Tick(AllottedGeometry, InCurrentTime, InDeltaTime);
diff --git a/Engine/Source/Editor/PackagesDialog/Private/SPackagesDialog.h b/Engine/Source/Editor/PackagesDialog/Private/SPackagesDialog.h
--- a/Engine/Source/Editor/PackagesDialog/Private/SPackagesDialog.h
+++ b/Engine/Source/Editor/PackagesDialog/Private/SPackagesDialog.h
@@ -507,18 +507,6 @@ private:
 	 */
 	void RefreshButtons();
 
-	/** 
-	 * @return	the desired toggle state for the ToggleSelectedCheckBox.
-	 * Returns Unchecked, unless all of the selected packages are Checked.
-	 */
-	ECheckBoxState GetToggleSelectedState() const;
-
-	/** 
-	 * Toggles the highlighted packages.
-	 * If no packages are explicitly highlighted, toggles all packages in the list.
-	 */
-	void OnToggleSelectedCheckBox(ECheckBoxState InNewState);
-
 	/** 
 	 * Makes the widget for the checkbox items in the list view 
 	 */
@@ -542,8 +530,8 @@ private:
 	/** 
 	 * Get all the selected items in the dialog.
 	 * 
-	 * @param bAllIfNone - if true, returns all teh items when none are selected
-	 * @return the array of packages to be altered by ToggleSelectedCheckBox
+	 * @param bAllIfNone - if true, returns all items when none are selected
+	 * @return the array of selected packages
 	 */
 	TArray< TSharedPtr<FPackageItem> > GetSelectedItems( bool bAllIfNone ) const;
 
@@ -585,9 +573,6 @@ private:
 	 */
 	void SortTree();
 
-	/** A Checkbox used to toggle multiple packages. */
-	TSharedPtr< SCheckBox > ToggleSelectedCheckBox;
-
 	/** All checkbox items stored in this widget for the list view */
 	TArray< TSharedPtr<FPackageItem> > Items;
 
