Fix submit dialog to output asset class and deleted color even for deleted objects without an asset type action.

diff --git a/Engine/Source/Editor/PackagesDialog/Private/SPackagesDialog.cpp b/Engine/Source/Editor/PackagesDialog/Private/SPackagesDialog.cpp
--- a/Engine/Source/Editor/PackagesDialog/Private/SPackagesDialog.cpp
+++ b/Engine/Source/Editor/PackagesDialog/Private/SPackagesDialog.cpp
@@ -103,6 +103,13 @@ bool FPackageItem::HasMultipleAssets() const
 
 bool FPackageItem::GetTypeNameAndColor(FText& OutName, FColor& OutColor) const
 {
+	if(HasMultipleAssets())
+	{
+		OutColor = FColor::White;
+		OutName = LOCTEXT("MultipleAssets", "Multiple Assets");
+		return true;
+	}
+
 	// Resolve the object belonging to the package and cache.
 	if (!Object.IsValid())
 	{
@@ -113,37 +120,39 @@ bool FPackageItem::GetTypeNameAndColor(FText& OutName, FColor& OutColor) const
 	{
 		// Load the asset tools module to get access to the class color
 		UObject* ObjectPtr = Object.Get(/*bEvenIfPendingKill*/true);
-		const FAssetToolsModule& AssetToolsModule = FModuleManager::LoadModuleChecked<FAssetToolsModule>(TEXT("AssetTools"));
-		const TSharedPtr<IAssetTypeActions> AssetTypeActions = AssetToolsModule.Get().GetAssetTypeActionsForClass(ObjectPtr->GetClass()).Pin();
-		if ( AssetTypeActions.IsValid() )
+
+		FAssetData AssetData(ObjectPtr);
+		OutName = FText::FromString(AssetData.AssetClassPath.ToString());
+
+		if (!IsValidChecked(ObjectPtr))
 		{
-			if(HasMultipleAssets())
+			OutColor = FColor::Red;
+		}
+		else
+		{
+			const FAssetToolsModule& AssetToolsModule = FModuleManager::LoadModuleChecked<FAssetToolsModule>(TEXT("AssetTools"));
+			const TSharedPtr<IAssetTypeActions> AssetTypeActions = AssetToolsModule.Get().GetAssetTypeActionsForClass(ObjectPtr->GetClass()).Pin();
+			if ( AssetTypeActions.IsValid() )
 			{
-				OutColor = FColor::White;
-				OutName = LOCTEXT("MultipleAssets", "Multiple Assets");
+				OutColor = AssetTypeActions->GetTypeColor();
 			}
-			else // Just one asset in the package.
+			else
 			{
-				OutColor = !IsValidChecked(ObjectPtr) ? FColor::Red : AssetTypeActions->GetTypeColor();
-
-				FAssetData AssetData(ObjectPtr);
-				OutName = FText::FromString(AssetData.AssetClassPath.ToString());
+				OutColor = FColor::White;
 			}
-			return true;
 		}
-	}
-	// if we do not find any package object, consider the package empty, return a red `Empty Package`
-	else
-	{
-		OutName = LOCTEXT("NoAssets", "Empty Package");
-		OutColor = FColor(					// Copied from ContentBrowserCLR.cpp
-			127 + FColor::Red.R / 2,		// Desaturate the colors a bit (GB colors were too.. much)
-			127 + FColor::Red.G / 2,
-			127 + FColor::Red.B / 2,
-			200);		// Opacity
 		return true;
 	}
-	return false;
+
+	// if we do not find any package object, consider the package empty, return a red `Empty Package`
+	OutName = LOCTEXT("NoAssets", "Empty Package");
+	OutColor = FColor(					// Copied from ContentBrowserCLR.cpp
+		127 + FColor::Red.R / 2,		// Desaturate the colors a bit (GB colors were too.. much)
+		127 + FColor::Red.G / 2,
+		127 + FColor::Red.B / 2,
+		200);		// Opacity
+
+	return true;
 }
 
 /**
