Update FPackageItem::GetTypeNameAndColor to consider empty packages as deleted packages by default, without the need of a pending kill objects in it.
This makes deleted entity packages show as deleted in the save dialogs when they are moved into a different package, when grouping then under another entity, for example.

diff --git a/Engine/Source/Editor/PackagesDialog/Private/SPackagesDialog.cpp b/Engine/Source/Editor/PackagesDialog/Private/SPackagesDialog.cpp
--- a/Engine/Source/Editor/PackagesDialog/Private/SPackagesDialog.cpp
+++ b/Engine/Source/Editor/PackagesDialog/Private/SPackagesDialog.cpp
@@ -124,7 +124,7 @@ bool FPackageItem::GetTypeNameAndColor(FText& OutName, FColor& OutColor) const
 		FAssetData AssetData(ObjectPtr);
 		OutName = FText::FromString(AssetData.AssetClassPath.ToString());
 
-		if (!IsValidChecked(ObjectPtr))
+		if (!IsValidChecked(ObjectPtr) || UPackage::IsEmptyPackage(GetPackage()))
 		{
 			OutColor = FColor::Red;
 		}
