Added support for showing Verse paths to the Packages Dialog.

diff --git a/Engine/Source/Editor/PackagesDialog/Private/PackagesDialog.cpp b/Engine/Source/Editor/PackagesDialog/Private/PackagesDialog.cpp
--- a/Engine/Source/Editor/PackagesDialog/Private/PackagesDialog.cpp
+++ b/Engine/Source/Editor/PackagesDialog/Private/PackagesDialog.cpp
@@ -207,10 +207,11 @@ void FPackagesDialogModule::RemoveAllPackageItems()
 void FPackagesDialogModule::AddPackageItem(UPackage* InPackage, ECheckBoxState InChecked, 
 	bool InDisabled/*=false*/, FString InIconName/*=TEXT("SavePackages.SCC_DlgNoIcon")*/, FString InIconToolTip/*=TEXT("")*/)
 {
-	FName AssetName;
-	FName OwnerName;
+	FString AssetDisplayName;
+	UE::Core::FVersePath VersePath;
+	FString OwnerName;
 
-	FAssetToolsModule& AssetToolsModule = FModuleManager::LoadModuleChecked<FAssetToolsModule>(TEXT("AssetTools"));
+	FAssetToolsModule& AssetToolsModule = FAssetToolsModule::GetModule();
 
 	// Lookup for the first asset in the package
 	UObject* FoundAsset = nullptr;
@@ -238,29 +239,34 @@ void FPackagesDialogModule::AddPackageItem(UPackage* InPackage, ECheckBoxState I
 		TWeakPtr<IAssetTypeActions> AssetTypeActions = AssetToolsModule.Get().GetAssetTypeActionsForClass(FoundAsset->GetClass());
 		if (AssetTypeActions.IsValid())
 		{
-			AssetName = *AssetTypeActions.Pin()->GetObjectDisplayName(FoundAsset);
+			AssetDisplayName = AssetTypeActions.Pin()->GetObjectDisplayName(FoundAsset);
 		}
 		else
 		{
-			AssetName = FoundAsset->GetFName();
+			AssetDisplayName = FoundAsset->GetName();
 		}
 
-		OwnerName = FoundAsset->GetOutermostObject()->GetFName();
+		if (AssetToolsModule.Get().ShowingContentVersePath())
+		{
+			VersePath = FoundAsset->GetVersePath();
+		}
+
+		OwnerName = FoundAsset->GetOutermostObject()->GetName();
 	}
 
 	// Last resort, display the package name
-	if (AssetName == NAME_None)
+	if (AssetDisplayName.IsEmpty())
 	{
-		AssetName = *FPackageName::GetShortName(InPackage->GetFName());
+		AssetDisplayName = FPackageName::GetShortName(InPackage->GetFName());
 	}
 
-	const FString Extension = (InPackage->GetPackageFlags() & PKG_ContainsMap) ? FPackageName::GetMapPackageExtension() : FPackageName::GetAssetPackageExtension();
+	const FString& Extension = (InPackage->GetPackageFlags() & PKG_ContainsMap) ? FPackageName::GetMapPackageExtension() : FPackageName::GetAssetPackageExtension();
 	FString FileName = FPaths::ConvertRelativePathToFull(FPackageName::LongPackageNameToFilename(InPackage->GetName(), Extension));
 
 	// Ensure that the filename is formatted in a way that the current platform natively understands. 
 	FPaths::MakePlatformFilename(FileName);
 	
-	PackagesDialogWidget.Get()->Add(MakeShareable(new FPackageItem(InPackage, AssetName.ToString(), FileName, OwnerName.ToString(), InChecked, InDisabled, InIconName, InIconToolTip)));
+	PackagesDialogWidget.Get()->Add(MakeShareable(new FPackageItem(InPackage, MoveTemp(AssetDisplayName), MoveTemp(FileName), MoveTemp(VersePath), MoveTemp(OwnerName), InChecked, InDisabled, MoveTemp(InIconName), MoveTemp(InIconToolTip))));
 }
 
 /**
diff --git a/Engine/Source/Editor/PackagesDialog/Private/SPackagesDialog.cpp b/Engine/Source/Editor/PackagesDialog/Private/SPackagesDialog.cpp
--- a/Engine/Source/Editor/PackagesDialog/Private/SPackagesDialog.cpp
+++ b/Engine/Source/Editor/PackagesDialog/Private/SPackagesDialog.cpp
@@ -564,21 +564,8 @@ TSharedRef<SWidget> SPackagesDialog::GenerateWidgetForItemAndColumn( TSharedPtr<
 {
 	check(Item.IsValid());
 
-	// Choose the icon based on the severity
-	const FSlateBrush* IconBrush = Item->GetIconName().IsEmpty() ? FStyleDefaults::GetNoBrush() : FAppStyle::GetBrush(*(Item->GetIconName()));
-
 	const FMargin RowPadding(3, 3, 3, 3);
 
-	// Extract the type and color for the package
-	FColor PackageColor;
-	FText PackageType;
-	Item->GetTypeNameAndColor(PackageType, PackageColor);
-
-	const FString AssetName = Item->GetAssetName();
-	const FString OwnerName = Item->GetOwnerName();
-	const FString PackageName = Item->GetPackageName();
-	const FString FileName = Item->GetFileName();
-
 	TSharedPtr<SWidget> ItemContentWidget;
 
 	if (ColumnID == SPackagesDialogDefs::ColumnID_CheckBoxLabel)
@@ -620,6 +607,9 @@ TSharedRef<SWidget> SPackagesDialog::GenerateWidgetForItemAndColumn( TSharedPtr<
 	}
 	else if (ColumnID == SPackagesDialogDefs::ColumnID_IconLabel)
 	{
+		// Choose the icon based on the severity
+		const FSlateBrush* IconBrush = Item->GetIconName().IsEmpty() ? FStyleDefaults::GetNoBrush() : FAppStyle::GetBrush(*(Item->GetIconName()));
+
 		ItemContentWidget = SNew(SHorizontalBox)
 			+SHorizontalBox::Slot()
 			.HAlign(HAlign_Center)
@@ -633,17 +623,21 @@ TSharedRef<SWidget> SPackagesDialog::GenerateWidgetForItemAndColumn( TSharedPtr<
 	}
 	else if (ColumnID == SPackagesDialogDefs::ColumnID_AssetLabel)
 	{
+		const FString& AssetDisplayName = Item->GetAssetDisplayName();
+
 		ItemContentWidget = SNew(SHorizontalBox)
 			+SHorizontalBox::Slot()
 			.Padding(RowPadding)
 			[
 				SNew(STextBlock)
-				.Text(FText::FromString(AssetName))
+				.Text(FText::FromString(AssetDisplayName))
 				.IsEnabled(!Item->IsDisabled())
 			];
 	}
 	else if (ColumnID == SPackagesDialogDefs::ColumnID_OwnerLabel)
 	{
+		const FString& OwnerName = Item->GetOwnerName();
+
 		ItemContentWidget = SNew(SHorizontalBox)
 			+SHorizontalBox::Slot()
 			.Padding(RowPadding)
@@ -656,18 +650,28 @@ TSharedRef<SWidget> SPackagesDialog::GenerateWidgetForItemAndColumn( TSharedPtr<
 	}
 	else if (ColumnID == SPackagesDialogDefs::ColumnID_PackageLabel)
 	{
+		const FString& PackageName = Item->GetPackageName();
+		const FString& FileName = Item->GetFileName();
+		const UE::Core::FVersePath& VersePath = Item->GetVersePath();
+
+		// Prefer the Verse path if valid.  It will be invalid if we are not showing Verse paths.
 		ItemContentWidget = SNew(SHorizontalBox)
 			+SHorizontalBox::Slot()
 			.Padding(RowPadding)
 			[
 				SNew(STextBlock)
-				.Text(FText::FromString(PackageName))
+				.Text(FText::FromString(VersePath.IsValid() ? VersePath.ToString() : PackageName))
 				.ToolTipText(FText::FromString(FileName))
 				.IsEnabled(!Item->IsDisabled())
 			];
 	}
 	else if (ColumnID == SPackagesDialogDefs::ColumnID_TypeLabel)
 	{
+		// Extract the type and color for the package
+		FColor PackageColor;
+		FText PackageType;
+		Item->GetTypeNameAndColor(PackageType, PackageColor);
+
 		ItemContentWidget = SNew(SHorizontalBox)
 			+SHorizontalBox::Slot()
 			.Padding(RowPadding)
@@ -725,6 +729,7 @@ TSharedPtr<SWidget> SPackagesDialog::MakePackageListContextMenu() const
 					FExecuteAction::CreateLambda( [SelectedItems] ()
 					{
 						TArray<FString> Paths;
+						Paths.Reserve(SelectedItems.Num());
 						for (const TSharedPtr<FPackageItem>& PackageItem: SelectedItems)
 						{
 							if (!PackageItem->GetFileName().IsEmpty())
@@ -732,7 +737,7 @@ TSharedPtr<SWidget> SPackagesDialog::MakePackageListContextMenu() const
 								Paths.Add(PackageItem->GetFileName());
 							}
 						}
-						FPlatformApplicationMisc::ClipboardCopy(*FString::Join(Paths, TEXT("\n")));
+						FPlatformApplicationMisc::ClipboardCopy(*FString::Join(Paths, LINE_TERMINATOR));
 					} ),
 					FCanExecuteAction::CreateLambda( [SelectedItems] ()
 					{
@@ -751,6 +756,7 @@ TSharedPtr<SWidget> SPackagesDialog::MakePackageListContextMenu() const
 					FExecuteAction::CreateLambda( [SelectedItems] ()
 					{
 						TArray<FString> Paths;
+						Paths.Reserve(SelectedItems.Num());
 						for (const TSharedPtr<FPackageItem>& PackageItem: SelectedItems)
 						{
 							if (!PackageItem->GetFileName().IsEmpty())
@@ -758,10 +764,42 @@ TSharedPtr<SWidget> SPackagesDialog::MakePackageListContextMenu() const
 								Paths.Add(PackageItem->GetPackageName());
 							}
 						}
-						FPlatformApplicationMisc::ClipboardCopy(*FString::Join(Paths, TEXT("\n")));
+						FPlatformApplicationMisc::ClipboardCopy(*FString::Join(Paths, LINE_TERMINATOR));
 					} )
 				)
 			);	
+
+			const FAssetToolsModule& AssetToolsModule = FAssetToolsModule::GetModule();
+			if (AssetToolsModule.Get().ShowingContentVersePath())
+			{
+				MenuBuilder.AddMenuEntry(
+					LOCTEXT("SCCCopyVersePathToClipboard", "Copy Verse Path"),
+					LOCTEXT("SCCCopyVersePathToClipboardTooltip", "Copies the Verse path to the clipboard."),
+					FSlateIcon(),
+					FUIAction(
+						FExecuteAction::CreateLambda([SelectedItems]()
+						{
+							TArray<FString> VersePaths;
+							VersePaths.Reserve(SelectedItems.Num());
+							for (const TSharedPtr<FPackageItem>& PackageItem: SelectedItems)
+							{
+								if (PackageItem->GetVersePath().IsValid())
+								{
+									VersePaths.Add(PackageItem->GetVersePath().ToString());
+								}
+							}
+							FPlatformApplicationMisc::ClipboardCopy(*FString::Join(VersePaths, LINE_TERMINATOR));
+						} ),
+						FCanExecuteAction::CreateLambda([SelectedItems]()
+						{
+							return Algo::AnyOf(SelectedItems, [](const TSharedPtr<FPackageItem>& PackageItem)
+							{
+								return PackageItem->GetVersePath().IsValid();
+							});
+						})
+					)
+				);
+			}
 		}
 		MenuBuilder.EndSection();
 	}
@@ -938,82 +976,82 @@ void SPackagesDialog::RequestSort()
 
 void SPackagesDialog::SortTree()
 {
+	TFunction<int32(const TSharedPtr<FPackageItem>&, const TSharedPtr<FPackageItem>&)> CompareColumn;
+
 	if (SortByColumn == SPackagesDialogDefs::ColumnID_AssetLabel)
 	{
-		if (SortMode == EColumnSortMode::Ascending)
+		CompareColumn = [](const TSharedPtr<FPackageItem>& A, const TSharedPtr<FPackageItem>& B)
 		{
-			Items.Sort([](const TSharedPtr<FPackageItem>& A, const TSharedPtr<FPackageItem>& B) {
-				return A->GetAssetName() < B->GetAssetName(); } );
-		}
-		else if (SortMode == EColumnSortMode::Descending)
-		{
-			Items.Sort([](const TSharedPtr<FPackageItem>& A, const TSharedPtr<FPackageItem>& B) {
-				return A->GetAssetName() >= B->GetAssetName(); } );
-		}
+			return A->GetAssetDisplayName().Compare(B->GetAssetDisplayName());
+		};
 	}
 	else if (SortByColumn == SPackagesDialogDefs::ColumnID_OwnerLabel)
 	{
-		if (SortMode == EColumnSortMode::Ascending)
+		CompareColumn = [](const TSharedPtr<FPackageItem>& A, const TSharedPtr<FPackageItem>& B)
 		{
-			Items.Sort([](const TSharedPtr<FPackageItem>& A, const TSharedPtr<FPackageItem>& B) {
-				return A->GetOwnerName() < B->GetOwnerName(); } );
-		}
-		else if (SortMode == EColumnSortMode::Descending)
-		{
-			Items.Sort([](const TSharedPtr<FPackageItem>& A, const TSharedPtr<FPackageItem>& B) {
-				return A->GetOwnerName() >= B->GetOwnerName(); } );
-		}
+			return A->GetOwnerName().Compare(B->GetOwnerName());
+		};
 	}
 	else if (SortByColumn == SPackagesDialogDefs::ColumnID_PackageLabel)
 	{
-		if (SortMode == EColumnSortMode::Ascending)
-		{
-			Items.Sort([](const TSharedPtr<FPackageItem>& A, const TSharedPtr<FPackageItem>& B) {
-				return A->GetPackageName() < B->GetPackageName(); } );
-		}
-		else if (SortMode == EColumnSortMode::Descending)
+		// Prefer the Verse path if valid.  It will be invalid if we are not showing Verse paths.
+		CompareColumn = [](const TSharedPtr<FPackageItem>& A, const TSharedPtr<FPackageItem>& B)
 		{
-			Items.Sort([](const TSharedPtr<FPackageItem>& A, const TSharedPtr<FPackageItem>& B) {
-				return A->GetPackageName() >= B->GetPackageName(); } );
-		}
+			if (A->GetVersePath().IsValid() && B->GetVersePath().IsValid())
+			{
+				return A->GetVersePath().Compare(B->GetVersePath());
+			}
+			else if (A->GetVersePath().IsValid())
+			{
+				return A->GetVersePath().ToString().Compare(B->GetPackageName());
+			}
+			else if (B->GetVersePath().IsValid())
+			{
+				return A->GetPackageName().Compare(B->GetVersePath().ToString());
+			}
+			else
+			{
+				return A->GetPackageName().Compare(B->GetPackageName());
+			}
+		};
 	}
 	else if (SortByColumn == SPackagesDialogDefs::ColumnID_TypeLabel)
 	{
-		if (SortMode == EColumnSortMode::Ascending)
+		CompareColumn = [](const TSharedPtr<FPackageItem>& A, const TSharedPtr<FPackageItem>& B)
 		{
-			Items.Sort([](const TSharedPtr<FPackageItem>& A, const TSharedPtr<FPackageItem>& B) {
-				return A->GetTypeName().CompareTo(B->GetTypeName()) < 0; } );
-		}
-		else if (SortMode == EColumnSortMode::Descending)
-		{
-			Items.Sort([](const TSharedPtr<FPackageItem>& A, const TSharedPtr<FPackageItem>& B) {
-				return A->GetTypeName().CompareTo(B->GetTypeName()) >= 0; } );
-		}
+			return A->GetTypeName().CompareTo(B->GetTypeName());
+		};
 	}
 	else if (SortByColumn == SPackagesDialogDefs::ColumnID_IconLabel)
 	{
-		if (SortMode == EColumnSortMode::Ascending)
-		{
-			Items.Sort([](const TSharedPtr<FPackageItem>& A, const TSharedPtr<FPackageItem>& B) {
-				return A->GetIconName() < B->GetIconName(); } );
-		}
-		else if (SortMode == EColumnSortMode::Descending)
+		CompareColumn = [](const TSharedPtr<FPackageItem>& A, const TSharedPtr<FPackageItem>& B)
 		{
-			Items.Sort([](const TSharedPtr<FPackageItem>& A, const TSharedPtr<FPackageItem>& B) {
-				return A->GetIconName() >= B->GetIconName(); } );
-		}
+			return A->GetIconName().Compare(B->GetIconName());
+		};
 	}
 	else if (SortByColumn == SPackagesDialogDefs::ColumnID_CheckedOutByLabel)
+	{
+		CompareColumn = [](const TSharedPtr<FPackageItem>& A, const TSharedPtr<FPackageItem>& B)
+		{
+			return A->GetCheckedOutByString().Compare(B->GetCheckedOutByString());
+		};
+	}
+
+	if (CompareColumn)
 	{
 		if (SortMode == EColumnSortMode::Ascending)
 		{
-			Items.Sort([](const TSharedPtr<FPackageItem>& A, const TSharedPtr<FPackageItem>& B) {
-				return A->GetCheckedOutByString() < B->GetCheckedOutByString(); });
+			Items.Sort([&CompareColumn](const TSharedPtr<FPackageItem>& A, const TSharedPtr<FPackageItem>& B)
+			{
+				return CompareColumn(A, B) < 0;
+			});
 		}
 		else if (SortMode == EColumnSortMode::Descending)
 		{
-			Items.Sort([](const TSharedPtr<FPackageItem>& A, const TSharedPtr<FPackageItem>& B) {
-				return A->GetCheckedOutByString() >= B->GetCheckedOutByString(); });
+			Items.Sort([&CompareColumn](const TSharedPtr<FPackageItem>& A, const TSharedPtr<FPackageItem>& B)
+			{
+				return CompareColumn(A, B) > 0;
+			});
 		}
 	}
 }
diff --git a/Engine/Source/Editor/PackagesDialog/Private/SPackagesDialog.h b/Engine/Source/Editor/PackagesDialog/Private/SPackagesDialog.h
--- a/Engine/Source/Editor/PackagesDialog/Private/SPackagesDialog.h
+++ b/Engine/Source/Editor/PackagesDialog/Private/SPackagesDialog.h
@@ -7,6 +7,7 @@
 #include "Containers/Set.h"
 #include "Containers/SparseArray.h"
 #include "Containers/UnrealString.h"
+#include "Containers/VersePath.h"
 #include "Delegates/Delegate.h"
 #include "HAL/PlatformCrt.h"
 #include "ISourceControlModule.h"
@@ -148,16 +149,17 @@ private:
 class FPackageItem : public TSharedFromThis<FPackageItem>
 {
 public:
-	FPackageItem(UPackage* InPackage, const FString& InAssetName, const FString& InFileName, const FString& InOwnerName, ECheckBoxState InState, bool InDisabled = false, FString InIconName=TEXT(""), FString InIconToolTip=TEXT(""))
+	FPackageItem(UPackage* InPackage, FString&& InAssetDisplayName, FString&& InFileName, UE::Core::FVersePath&& InVersePath, FString&& InOwnerName, ECheckBoxState InState, bool InDisabled, FString&& InIconName, FString&& InIconToolTip)
 		: Package(InPackage)
-		, AssetName(InAssetName)
+		, AssetDisplayName(MoveTemp(InAssetDisplayName))
 		, PackageName(FPackageName::ObjectPathToPackageName(InPackage->GetName()))
-		, FileName(InFileName)
-		, OwnerName(InOwnerName)
+		, FileName(MoveTemp(InFileName))
+		, VersePath(MoveTemp(InVersePath))
+		, OwnerName(MoveTemp(InOwnerName))
 		, State(InState)
 		, Disabled(InDisabled)
-		, IconName(InIconName)
-		, IconToolTip(InIconToolTip)
+		, IconName(MoveTemp(InIconName))
+		, IconToolTip(MoveTemp(InIconToolTip))
 	{
 		/** if the item is checked and disabled make the state Undetermined */
 		if(State == ECheckBoxState::Checked && Disabled)
@@ -240,11 +242,11 @@ public:
 	bool IsDisabled() const { return Disabled; }
 
 	/**
-	 * Gets the name of the asset item
+	 * Gets the display name of the primary asset item
 	 *
-	 * @return the name of the asset item
+	 * @return the display name of the asset item
 	 */
-	const FString& GetAssetName() const { return AssetName; }
+	const FString& GetAssetDisplayName() const { return AssetDisplayName; }
 
 	/**
 	 * Gets the name of the package item
@@ -260,6 +262,13 @@ public:
 	 */
 	const FString& GetFileName() const { return FileName; }
 
+	/**
+	 * Gets the Verse path of the primary asset item
+	 *
+	 * @return the Verse path of the primary asset item
+	 */
+	const UE::Core::FVersePath& GetVersePath() const { return VersePath; }
+
 	/**
 	 * Gets the name of the file owner item
 	 *
@@ -353,9 +362,10 @@ public:
 
 private:
 	UPackage* Package; 						// The package associated with this entry
-	FString AssetName;						// Name of the asset to display
+	FString AssetDisplayName;				// Name of the asset to display
 	FString PackageName;					// Name of the package to display
 	FString FileName;						// Name of the file
+	UE::Core::FVersePath VersePath;			// Verse path to display
 	FString OwnerName;						// Name of the owner of the file
 	ECheckBoxState State;					// The state of the checkbox
 	bool Disabled; 							// if the entry is disabled
