Added asset name and package path as tie breakers to packages dialog column sorting.  Otherwise you get a fairly random results when sorting a long list by columns that tend to have the same value.

diff --git a/Engine/Source/Editor/PackagesDialog/Private/SPackagesDialog.cpp b/Engine/Source/Editor/PackagesDialog/Private/SPackagesDialog.cpp
--- a/Engine/Source/Editor/PackagesDialog/Private/SPackagesDialog.cpp
+++ b/Engine/Source/Editor/PackagesDialog/Private/SPackagesDialog.cpp
@@ -976,84 +976,113 @@ void SPackagesDialog::RequestSort()
 
 void SPackagesDialog::SortTree()
 {
-	TFunction<int32(const TSharedPtr<FPackageItem>&, const TSharedPtr<FPackageItem>&)> CompareColumn;
-
-	if (SortByColumn == SPackagesDialogDefs::ColumnID_AssetLabel)
+	int32 Sign;
+	if (SortMode == EColumnSortMode::Ascending)
 	{
-		CompareColumn = [](const TSharedPtr<FPackageItem>& A, const TSharedPtr<FPackageItem>& B)
-		{
-			return A->GetAssetDisplayName().Compare(B->GetAssetDisplayName());
-		};
+		Sign = 1;
 	}
-	else if (SortByColumn == SPackagesDialogDefs::ColumnID_OwnerLabel)
+	else if (SortMode == EColumnSortMode::Descending)
 	{
-		CompareColumn = [](const TSharedPtr<FPackageItem>& A, const TSharedPtr<FPackageItem>& B)
-		{
-			return A->GetOwnerName().Compare(B->GetOwnerName());
-		};
+		Sign = -1;
 	}
-	else if (SortByColumn == SPackagesDialogDefs::ColumnID_PackageLabel)
+	else
 	{
-		// Prefer the Verse path if valid.  It will be invalid if we are not showing Verse paths.
-		CompareColumn = [](const TSharedPtr<FPackageItem>& A, const TSharedPtr<FPackageItem>& B)
-		{
-			if (A->GetVersePath().IsValid() && B->GetVersePath().IsValid())
-			{
-				return A->GetVersePath().Compare(B->GetVersePath());
-			}
-			else if (A->GetVersePath().IsValid())
-			{
-				return A->GetVersePath().ToString().Compare(B->GetPackageName());
-			}
-			else if (B->GetVersePath().IsValid())
-			{
-				return A->GetPackageName().Compare(B->GetVersePath().ToString());
-			}
-			else
-			{
-				return A->GetPackageName().Compare(B->GetPackageName());
-			}
-		};
+		return;
 	}
-	else if (SortByColumn == SPackagesDialogDefs::ColumnID_TypeLabel)
+
+	constexpr int32 MaxComparers = 3;
+	TArray<TFunction<int32(const TSharedPtr<FPackageItem>&, const TSharedPtr<FPackageItem>&)>, TInlineAllocator<MaxComparers>> CompareColumns;
+	CompareColumns.Reserve(MaxComparers);
+
+	auto CompareAssetColumn = [](const TSharedPtr<FPackageItem>& A, const TSharedPtr<FPackageItem>& B)
 	{
-		CompareColumn = [](const TSharedPtr<FPackageItem>& A, const TSharedPtr<FPackageItem>& B)
-		{
-			return A->GetTypeName().CompareTo(B->GetTypeName());
-		};
-	}
-	else if (SortByColumn == SPackagesDialogDefs::ColumnID_IconLabel)
+		return A->GetAssetDisplayName().Compare(B->GetAssetDisplayName());
+	};
+
+	// Prefer the Verse path if valid.  It will be invalid if we are not showing Verse paths.
+	auto ComparePackageColumn = [](const TSharedPtr<FPackageItem>& A, const TSharedPtr<FPackageItem>& B)
 	{
-		CompareColumn = [](const TSharedPtr<FPackageItem>& A, const TSharedPtr<FPackageItem>& B)
+		if (A->GetVersePath().IsValid() && B->GetVersePath().IsValid())
+		{
+			return A->GetVersePath().Compare(B->GetVersePath());
+		}
+		else if (A->GetVersePath().IsValid())
+		{
+			return A->GetVersePath().ToString().Compare(B->GetPackageName());
+		}
+		else if (B->GetVersePath().IsValid())
+		{
+			return A->GetPackageName().Compare(B->GetVersePath().ToString());
+		}
+		else
 		{
-			return A->GetIconName().Compare(B->GetIconName());
-		};
+			return A->GetPackageName().Compare(B->GetPackageName());
+		}
+	};
+
+	if (SortByColumn == SPackagesDialogDefs::ColumnID_AssetLabel)
+	{
+		CompareColumns.Emplace(MoveTemp(CompareAssetColumn));
+
+		// Use package path as tie breaker.
+		CompareColumns.Emplace(MoveTemp(ComparePackageColumn));
 	}
-	else if (SortByColumn == SPackagesDialogDefs::ColumnID_CheckedOutByLabel)
+	else if (SortByColumn == SPackagesDialogDefs::ColumnID_PackageLabel)
 	{
-		CompareColumn = [](const TSharedPtr<FPackageItem>& A, const TSharedPtr<FPackageItem>& B)
-		{
-			return A->GetCheckedOutByString().Compare(B->GetCheckedOutByString());
-		};
+		CompareColumns.Emplace(MoveTemp(ComparePackageColumn));
 	}
-
-	if (CompareColumn)
+	else
 	{
-		if (SortMode == EColumnSortMode::Ascending)
+		if (SortByColumn == SPackagesDialogDefs::ColumnID_IconLabel)
 		{
-			Items.Sort([&CompareColumn](const TSharedPtr<FPackageItem>& A, const TSharedPtr<FPackageItem>& B)
+			CompareColumns.Emplace([](const TSharedPtr<FPackageItem>& A, const TSharedPtr<FPackageItem>& B)
 			{
-				return CompareColumn(A, B) < 0;
+				return A->GetIconName().Compare(B->GetIconName());
 			});
 		}
-		else if (SortMode == EColumnSortMode::Descending)
+		else if (SortByColumn == SPackagesDialogDefs::ColumnID_OwnerLabel)
 		{
-			Items.Sort([&CompareColumn](const TSharedPtr<FPackageItem>& A, const TSharedPtr<FPackageItem>& B)
+			CompareColumns.Emplace([](const TSharedPtr<FPackageItem>& A, const TSharedPtr<FPackageItem>& B)
 			{
-				return CompareColumn(A, B) > 0;
+				return A->GetOwnerName().Compare(B->GetOwnerName());
 			});
 		}
+		else if (SortByColumn == SPackagesDialogDefs::ColumnID_TypeLabel)
+		{
+			CompareColumns.Emplace([](const TSharedPtr<FPackageItem>& A, const TSharedPtr<FPackageItem>& B)
+			{
+				return A->GetTypeName().CompareTo(B->GetTypeName());
+			});
+		}
+		else if (SortByColumn == SPackagesDialogDefs::ColumnID_CheckedOutByLabel)
+		{
+			CompareColumns.Emplace([](const TSharedPtr<FPackageItem>& A, const TSharedPtr<FPackageItem>& B)
+			{
+				return A->GetCheckedOutByString().Compare(B->GetCheckedOutByString());
+			});
+		}
+		else
+		{
+			return;
+		}
+
+		// Use asset name and package path as tie breakers.
+		CompareColumns.Emplace(MoveTemp(CompareAssetColumn));
+		CompareColumns.Emplace(MoveTemp(ComparePackageColumn));
 	}
+
+	Items.Sort([&CompareColumns, Sign](const TSharedPtr<FPackageItem>& A, const TSharedPtr<FPackageItem>& B)
+	{
+		for (const TFunction<int32(const TSharedPtr<FPackageItem>&, const TSharedPtr<FPackageItem>&)>& CompareColumn : CompareColumns)
+		{
+			const int32 Result = CompareColumn(A, B);
+			if (Result != 0)
+			{
+				return Sign * Result < 0;
+			}
+		}
+		return false;
+	});
 }
 
 void SPackageItemsListRow::Construct( const FArguments& InArgs, const TSharedRef<STableViewBase>& InOwnerTableView )
