OFPE: fix save dialog to use UPackage::FindAssetInPackage instead of existing code in order to support packages with multiples (garbage) assets.
Fix save dialog showing VERSE_DEAD objects as primary assets when reinstancing.

diff --git a/Engine/Source/Editor/PackagesDialog/Private/PackagesDialog.cpp b/Engine/Source/Editor/PackagesDialog/Private/PackagesDialog.cpp
--- a/Engine/Source/Editor/PackagesDialog/Private/PackagesDialog.cpp
+++ b/Engine/Source/Editor/PackagesDialog/Private/PackagesDialog.cpp
@@ -213,28 +213,8 @@ void FPackagesDialogModule::AddPackageItem(UPackage* InPackage, ECheckBoxState I
 
 	FAssetToolsModule& AssetToolsModule = FAssetToolsModule::GetModule();
 
-	// Lookup for the first asset in the package
-	UObject* FoundAsset = nullptr;
-	ForEachObjectWithPackage(InPackage, [&FoundAsset](UObject* InnerObject)
-	{
-		if (InnerObject->IsAsset())
-		{
-			if (FAssetData::IsUAsset(InnerObject))
-			{
-				// If we found the primary asset, use it
-				FoundAsset = InnerObject;
-				return false;
-			}
-			// Otherwise, keep the first found asset but keep looking for a primary asset
-			if (!FoundAsset)
-			{
-				FoundAsset = InnerObject;
-			}
-		}
-		return true;
-	}, /*bIncludeNestedObjects*/ false);
-
-	if (FoundAsset)
+	// Lookup for the asset in the package
+	if (UObject* FoundAsset = InPackage->FindAssetInPackage())
 	{
 		TWeakPtr<IAssetTypeActions> AssetTypeActions = AssetToolsModule.Get().GetAssetTypeActionsForClass(FoundAsset->GetClass());
 		if (AssetTypeActions.IsValid())
