Tentative fix to prevent crash when SWindow is nullptr in Package Dialog.

diff --git a/Engine/Source/Editor/PackagesDialog/Private/PackagesDialog.cpp b/Engine/Source/Editor/PackagesDialog/Private/PackagesDialog.cpp
--- a/Engine/Source/Editor/PackagesDialog/Private/PackagesDialog.cpp
+++ b/Engine/Source/Editor/PackagesDialog/Private/PackagesDialog.cpp
@@ -136,11 +136,15 @@ EDialogReturnType FPackagesDialogModule::ShowPackagesDialog(OUT TSet<FString>& I
  */
 void FPackagesDialogModule::RemovePackagesDialog()
 {
+	const TSharedPtr<SWindow> PinnedWindow = EditorPackagesDialogWindow.Pin();
 	/** Ensure the package dialog has been created already */
-	check( EditorPackagesDialogWindow.IsValid() );
+	check( PinnedWindow.IsValid() );
 
 	/** Request to destroy the package dialog window */
-	EditorPackagesDialogWindow.Pin()->RequestDestroyWindow();
+	if (PinnedWindow)
+	{
+		PinnedWindow->RequestDestroyWindow();
+	}
 }
 
 void FPackagesDialogModule::SetMessage(const FText& InMessage)
diff --git a/Engine/Source/Editor/PackagesDialog/Private/SPackagesDialog.h b/Engine/Source/Editor/PackagesDialog/Private/SPackagesDialog.h
--- a/Engine/Source/Editor/PackagesDialog/Private/SPackagesDialog.h
+++ b/Engine/Source/Editor/PackagesDialog/Private/SPackagesDialog.h
@@ -66,7 +66,10 @@ public:
 	FReply OnButtonClicked()
 	{
  		Clicked = true;
-		Module->RemovePackagesDialog();
+		if (Module)
+		{
+			Module->RemovePackagesDialog();
+		}
 		return FReply::Handled();
 	}
 
