EditorUsability : PackagesDialog
- Added asset class icon with color to the save all dialog content table row

#Synced-CL 51264175

diff --git a/Engine/Source/Editor/PackagesDialog/Private/SPackagesDialog.cpp b/Engine/Source/Editor/PackagesDialog/Private/SPackagesDialog.cpp
--- a/Engine/Source/Editor/PackagesDialog/Private/SPackagesDialog.cpp
+++ b/Engine/Source/Editor/PackagesDialog/Private/SPackagesDialog.cpp
@@ -624,10 +624,33 @@ TSharedRef<SWidget> SPackagesDialog::GenerateWidgetForItemAndColumn( TSharedPtr<
 	else if (ColumnID == SPackagesDialogDefs::ColumnID_AssetLabel)
 	{
 		const FString& AssetDisplayName = Item->GetAssetDisplayName();
+		const FSlateIcon& AssetIcon = Item->GetAssetIcon();
+		FText TypeName;
+		FColor TypeColor;
+		Item->GetTypeNameAndColor(TypeName, TypeColor);
 
 		ItemContentWidget = SNew(SHorizontalBox)
+			+SHorizontalBox::Slot()
+			.AutoWidth()
+			.HAlign(HAlign_Center)
+			.VAlign(VAlign_Center)
+			.Padding(RowPadding)
+			[
+				SNew(SBox)
+				.MaxAspectRatio(1.f)
+				.HeightOverride(16.f)
+				.WidthOverride(16.f)
+				[
+					SNew(SImage)
+					.ToolTipText(TypeName)
+					.ColorAndOpacity(TypeColor)
+					.Image(AssetIcon.GetIcon())
+					.IsEnabled(!Item->IsDisabled())
+				]
+			]
 			+SHorizontalBox::Slot()
 			.Padding(RowPadding)
+			.FillWidth(1.f)
 			[
 				SNew(STextBlock)
 				.Text(FText::FromString(AssetDisplayName))
@@ -667,10 +690,7 @@ TSharedRef<SWidget> SPackagesDialog::GenerateWidgetForItemAndColumn( TSharedPtr<
 	}
 	else if (ColumnID == SPackagesDialogDefs::ColumnID_TypeLabel)
 	{
-		// Extract the type and color for the package
-		FColor PackageColor;
-		FText PackageType;
-		Item->GetTypeNameAndColor(PackageType, PackageColor);
+		const FText PackageType = Item->GetTypeName();
 
 		ItemContentWidget = SNew(SHorizontalBox)
 			+SHorizontalBox::Slot()
@@ -680,7 +700,6 @@ TSharedRef<SWidget> SPackagesDialog::GenerateWidgetForItemAndColumn( TSharedPtr<
 				.Text(PackageType)
 				.ToolTipText(PackageType)
 				.IsEnabled(!Item->IsDisabled())
-				.ColorAndOpacity(PackageColor)
 			];
 	}
 	else if (ColumnID == SPackagesDialogDefs::ColumnID_CheckedOutByLabel)
diff --git a/Engine/Source/Editor/PackagesDialog/Private/SPackagesDialog.h b/Engine/Source/Editor/PackagesDialog/Private/SPackagesDialog.h
--- a/Engine/Source/Editor/PackagesDialog/Private/SPackagesDialog.h
+++ b/Engine/Source/Editor/PackagesDialog/Private/SPackagesDialog.h
@@ -21,6 +21,7 @@
 #include "Misc/Optional.h"
 #include "Misc/PackageName.h"
 #include "PackagesDialog.h"
+#include "Styling/SlateIconFinder.h"
 #include "Styling/SlateTypes.h"
 #include "Templates/SharedPointer.h"
 #include "Templates/TypeHash.h"
@@ -169,8 +170,25 @@ public:
 		{
 			State = ECheckBoxState::Undetermined;
 		}
+		
+		// Find the appropriate icon
+		if (const UObject* AssetObject = GetPackageObject())
+		{
+			AssetIcon = FSlateIconFinder::FindIconForClass(AssetObject->GetClass());
+		}
+
+		// Falling back on default icon
+		if (!AssetIcon.IsSet())
+		{
+			AssetIcon = FSlateIconFinder::FindIconForClass(UObject::StaticClass());
+		}
 	}
 
+	/**
+	 * @return FSlateIcon the icon for the underlying asset object class or default object icon
+	 */
+	const FSlateIcon& GetAssetIcon() const { return AssetIcon; }
+
 	/**
 	 * Gets the display state of the item
 	 *
@@ -376,6 +394,9 @@ private:
 	FString IconToolTip;					// ToolTip to display for the icon
 	FSimpleDelegate RefreshButtonCallback;	// ToolTip to display for the icon
 	mutable TWeakObjectPtr<UObject> Object;	// Cached object associated with this entry.
+
+	/** Icon for the asset class */
+	FSlateIcon AssetIcon;
 };
 
 /**
